Patches for Visual Studio workspace support in old versions

diff --git a/Engine/Extras/VisualStudioWorkspace/UnrealBuildTool-4.27.patch b/Engine/Extras/VisualStudioWorkspace/UnrealBuildTool-4.27.patch
new file mode 100644
--- /dev/null
+++ b/Engine/Extras/VisualStudioWorkspace/UnrealBuildTool-4.27.patch
@@ -0,0 +1,870 @@
+From 1ef457067fb87ae7c2317572080b76bd4435c4ac Mon Sep 17 00:00:00 2001
+From: [REDACTED_PERSON] <[REDACTED_EMAIL]>
+Date: Wed, 17 Jan 2024 19:49:54 -0800
+Subject: [PATCH] [UE4.27] Project Generator for VS Workspace support
+
+Adding a new Project Generator for Visual Studio that will be used for
+the new support to opening `.uproject` files directly in the IDE without
+generating the traditional `.sln` and .vcxproj files.
+
+This is based on the QueryMode API in `ue5-main` and adapted to the
+GenerateProjectFiles flow. The project files will contain information for each
+Target+Configuration+Platform combinations that are used the project.
+
+Visual Studio will automatically invoke this new generator during the flow to
+open a `.uproject` file directly.
+---
+ .../Configuration/UEBuildTarget.cs            |   2 +-
+ .../Modes/GenerateProjectFilesMode.cs         |   4 +
+ .../Platform/Windows/VCToolChain.cs           |  26 ++
+ .../ProjectFiles/ProjectFileGenerator.cs      |   5 +-
+ .../VSWorkspaceProjectFile.cs                 | 388 ++++++++++++++++++
+ .../VSWorkspaceProjectFileGenerator.cs        | 277 +++++++++++++
+ .../UnrealBuildTool/ToolChain/UEToolChain.cs  |  25 ++
+ .../UnrealBuildTool/UnrealBuildTool.csproj    |   2 +
+ 8 files changed, 726 insertions(+), 3 deletions(-)
+ create mode 100644 Engine/Source/Programs/UnrealBuildTool/ProjectFiles/VisualStudioWorkspace/VSWorkspaceProjectFile.cs
+ create mode 100644 Engine/Source/Programs/UnrealBuildTool/ProjectFiles/VisualStudioWorkspace/VSWorkspaceProjectFileGenerator.cs
+
+diff --git a/Engine/Source/Programs/UnrealBuildTool/Configuration/UEBuildTarget.cs b/Engine/Source/Programs/UnrealBuildTool/Configuration/UEBuildTarget.cs
+index 97a937d41af1b..947ac655a516e 100644
+--- a/Engine/Source/Programs/UnrealBuildTool/Configuration/UEBuildTarget.cs
++++ b/Engine/Source/Programs/UnrealBuildTool/Configuration/UEBuildTarget.cs
+@@ -2098,7 +2098,7 @@ List<RestrictedFolder> GetRestrictedFolders(FileReference File)
+ 		/// Creates a toolchain for the current target. May be overridden by the target rules.
+ 		/// </summary>
+ 		/// <returns>New toolchain instance</returns>
+-		private UEToolChain CreateToolchain(UnrealTargetPlatform Platform)
++		public UEToolChain CreateToolchain(UnrealTargetPlatform Platform)
+ 		{
+ 			if (Rules.ToolChainName == null)
+ 			{
+diff --git a/Engine/Source/Programs/UnrealBuildTool/Modes/GenerateProjectFilesMode.cs b/Engine/Source/Programs/UnrealBuildTool/Modes/GenerateProjectFilesMode.cs
+index a42da4931c523..f6d05618f8ff0 100644
+--- a/Engine/Source/Programs/UnrealBuildTool/Modes/GenerateProjectFilesMode.cs
++++ b/Engine/Source/Programs/UnrealBuildTool/Modes/GenerateProjectFilesMode.cs
+@@ -36,6 +36,7 @@ class GenerateProjectFilesMode : ToolMode
+ 		[CommandLine("-EddieProjectFiles", Value = nameof(ProjectFileFormat.Eddie))]
+ 		[CommandLine("-VSCode", Value = nameof(ProjectFileFormat.VisualStudioCode))]
+ 		[CommandLine("-VSMac", Value = nameof(ProjectFileFormat.VisualStudioMac))]
++		[CommandLine("-VSWorkspace", Value = nameof(ProjectFileFormat.VisualStudioWorkspace))]
+ 		[CommandLine("-CLion", Value = nameof(ProjectFileFormat.CLion))]
+ 		[CommandLine("-Rider", Value = nameof(ProjectFileFormat.Rider))]
+ 		HashSet<ProjectFileFormat> ProjectFileFormats = new HashSet<ProjectFileFormat>();
+@@ -184,6 +185,9 @@ public override int Execute(CommandLineArguments Arguments)
+ 					case ProjectFileFormat.VisualStudioCode:
+ 						Generator = new VSCodeProjectFileGenerator(ProjectFile);
+ 						break;
++					case ProjectFileFormat.VisualStudioWorkspace:
++						Generator = new VSWorkspaceProjectFileGenerator(ProjectFile, Arguments);
++						break;
+ 					case ProjectFileFormat.CLion:
+ 						Generator = new CLionGenerator(ProjectFile);
+ 						break;
+diff --git a/Engine/Source/Programs/UnrealBuildTool/Platform/Windows/VCToolChain.cs b/Engine/Source/Programs/UnrealBuildTool/Platform/Windows/VCToolChain.cs
+index e526fa7ee2e4c..9e776964c8431 100644
+--- a/Engine/Source/Programs/UnrealBuildTool/Platform/Windows/VCToolChain.cs
++++ b/Engine/Source/Programs/UnrealBuildTool/Platform/Windows/VCToolChain.cs
+@@ -46,6 +46,11 @@ public VCToolChain(ReadOnlyTargetRules Target)
+ 			}
+ 		}
+ 
++		public override FileReference GetCppCompilerPath()
++		{
++			return EnvVars.CompilerPath;
++		}
++
+ 		/// <summary>
+ 		/// Prepares the environment for building
+ 		/// </summary>
+@@ -151,6 +156,27 @@ public static void AddSystemIncludePath(List<string> Arguments, DirectoryReferen
+ 		}
+ 
+ 
++		public override IEnumerable<string> GetGlobalCommandLineArgs(CppCompileEnvironment CompileEnvironment)
++		{
++			List<string> Arguments = new List<string>();
++			AppendCLArguments_Global(new CppCompileEnvironment(CompileEnvironment), Arguments);
++			return Arguments;
++		}
++
++		public override IEnumerable<string> GetCPPCommandLineArgs(CppCompileEnvironment CompileEnvironment)
++		{
++			List<string> Arguments = new List<string>();
++			AppendCLArguments_CPP(CompileEnvironment, Arguments);
++			return Arguments;
++		}
++
++		public override IEnumerable<string> GetCCommandLineArgs(CppCompileEnvironment CompileEnvironment)
++		{
++			List<string> Arguments = new List<string>();
++			AppendCLArguments_C(Arguments);
++			return Arguments;
++		}
++
+ 		protected virtual void AppendCLArguments_Global(CppCompileEnvironment CompileEnvironment, List<string> Arguments)
+ 		{
+ 			// Suppress generation of object code for unreferenced inline functions. Enabling this option is more standards compliant, and causes a big reduction
+diff --git a/Engine/Source/Programs/UnrealBuildTool/ProjectFiles/ProjectFileGenerator.cs b/Engine/Source/Programs/UnrealBuildTool/ProjectFiles/ProjectFileGenerator.cs
+index 66a1ec6d0b227..f64a010b97581 100644
+--- a/Engine/Source/Programs/UnrealBuildTool/ProjectFiles/ProjectFileGenerator.cs
++++ b/Engine/Source/Programs/UnrealBuildTool/ProjectFiles/ProjectFileGenerator.cs
+@@ -135,6 +135,7 @@ enum ProjectFileFormat
+ 		VisualStudio2017,
+ 		VisualStudio2019,
+ 		VisualStudio2022,
++		VisualStudioWorkspace,
+ 		XCode,
+ 		Eddie,
+ 		VisualStudioCode,
+@@ -2234,7 +2235,7 @@ private ProjectFile FindProjectForModule(FileReference CurModuleFile, List<FileR
+ 		/// <param name="EnterpriseProject">The enterprise project we created</param>
+ 		/// <param name="GameProjects">Map of game folder name to all of the game projects we created</param>
+ 		/// <param name="ProgramProjects">Map of program names to all of the program projects we created</param>
+-		private void AddProjectsForAllTargets(
++		protected void AddProjectsForAllTargets(
+ 			PlatformProjectGeneratorCollection PlatformProjectGenerators,
+ 			List<FileReference> AllGames,
+ 			List<FileReference> AllTargetFiles,
+@@ -2586,7 +2587,7 @@ protected void AddProjectsForMods(List<ProjectFile> GameProjects, out List<Proje
+ 		/// </summary>
+ 		/// <param name="BaseName">The base name for the project file</param>
+ 		/// <returns>Full path to the project file</returns>
+-		protected FileReference GetProjectLocation(string BaseName)
++		protected virtual FileReference GetProjectLocation(string BaseName)
+ 		{
+ 			return FileReference.Combine(IntermediateProjectFilesPath, BaseName + ProjectFileExtension);
+ 		}
+diff --git a/Engine/Source/Programs/UnrealBuildTool/ProjectFiles/VisualStudioWorkspace/VSWorkspaceProjectFile.cs b/Engine/Source/Programs/UnrealBuildTool/ProjectFiles/VisualStudioWorkspace/VSWorkspaceProjectFile.cs
+new file mode 100644
+index 0000000000000..2e37959f34c0a
+--- /dev/null
++++ b/Engine/Source/Programs/UnrealBuildTool/ProjectFiles/VisualStudioWorkspace/VSWorkspaceProjectFile.cs
+@@ -0,0 +1,388 @@
++// Copyright Epic Games, Inc. All Rights Reserved.
++
++using System;
++using System.Collections.Generic;
++using System.IO;
++using System.Linq;
++using System.Text;
++using Tools.DotNETCommon;
++
++namespace UnrealBuildTool
++{
++	internal class VSWorkspaceProjectFile : ProjectFile
++	{
++		private readonly DirectoryReference RootPath;
++		private readonly HashSet<TargetType> TargetTypes;
++		private readonly CommandLineArguments Arguments;
++		private VCProjectFileSettings Settings;
++		private bool bUsePrecompiled;
++
++		/// <summary>
++		/// Collection of output files for this project
++		/// </summary>
++		public List<ExportedTargetInfo> ExportedTargetProjects { get; set; } = new List<ExportedTargetInfo>();
++
++		public VSWorkspaceProjectFile(FileReference InProjectFilePath,
++			DirectoryReference RootPath, HashSet<TargetType> TargetTypes, CommandLineArguments Arguments, VCProjectFileSettings Settings, bool bUsePrecompiled)
++			: base(InProjectFilePath)
++		{
++			this.RootPath = RootPath;
++			this.TargetTypes = TargetTypes;
++			this.Arguments = Arguments;
++			this.Settings = Settings;
++			this.bUsePrecompiled = bUsePrecompiled;
++		}
++
++		/// <summary>
++		/// Write project file info in JSON file.
++		/// For every combination of <c>UnrealTargetPlatform</c>, <c>UnrealTargetConfiguration</c> and <c>TargetType</c>
++		/// will be generated separate JSON file.
++		/// </summary>
++		public bool WriteProjectFile(List<UnrealTargetPlatform> InPlatforms,
++			List<UnrealTargetConfiguration> InConfigurations,
++			PlatformProjectGeneratorCollection PlatformProjectGenerators, JsonWriterStyle Minimize)
++		{
++			DirectoryReference ProjectRootFolder = RootPath;
++
++			Dictionary<FileReference, Tuple<UEBuildTarget, bool>> FileToTarget = new Dictionary<FileReference, Tuple<UEBuildTarget, bool>>();
++
++			foreach (UnrealTargetPlatform Platform in InPlatforms)
++			{
++				foreach (UnrealTargetConfiguration Configuration in InConfigurations)
++				{
++					foreach (ProjectTarget ProjectTarget in ProjectTargets.OfType<ProjectTarget>())
++					{
++						if (TargetTypes.Any() && !TargetTypes.Contains(ProjectTarget.TargetRules.Type))
++						{
++							continue;
++						}
++
++						// Skip Programs for all configs except for current platform + Development & Debug configurations
++						if (ProjectTarget.TargetRules.Type == TargetType.Program &&
++							(BuildHostPlatform.Current.Platform != Platform ||
++							 !(Configuration == UnrealTargetConfiguration.Development || Configuration == UnrealTargetConfiguration.Debug)))
++						{
++							continue;
++						}
++
++						// Skip Editor for all platforms except for current platform
++						if (ProjectTarget.TargetRules.Type == TargetType.Editor && 
++							(BuildHostPlatform.Current.Platform != Platform || 
++							(Configuration == UnrealTargetConfiguration.Test || Configuration == UnrealTargetConfiguration.Shipping)))
++						{
++							continue;
++						}
++
++						bool bBuildByDefault = ShouldBuildByDefaultForSolutionTargets && ProjectTarget.SupportedPlatforms.Contains(Platform);
++
++						string DefaultArchitecture = UEBuildPlatform
++							.GetBuildPlatform(Platform)
++							.GetDefaultArchitecture(ProjectTarget.UnrealProjectFilePath);
++
++						TargetDescriptor TargetDesc = new TargetDescriptor(ProjectTarget.UnrealProjectFilePath, ProjectTarget.Name,
++							Platform, Configuration, DefaultArchitecture, Arguments);
++
++						try
++						{
++							FileReference OutputFile = FileReference.Combine(ProjectRootFolder, 
++								$"{ProjectTarget.TargetFilePath.GetFileNameWithoutAnyExtensions()}_{Configuration}_{Platform}.json");
++
++							UEBuildTarget BuildTarget = UEBuildTarget.Create(TargetDesc, false, false);
++							FileToTarget.Add(OutputFile, new Tuple<UEBuildTarget, bool>(BuildTarget, bBuildByDefault));
++						}
++						catch (Exception Ex)
++						{
++							Log.TraceWarning("Exception while generating include data for Target:{0}, Platform: {1}, Configuration: {2}", 
++								TargetDesc.Name, Platform.ToString(), Configuration.ToString());
++							Log.TraceWarning(Ex.ToString());
++						}
++					}
++				}
++			}
++
++			foreach (var Entry in FileToTarget)
++			{
++				var OutputFile = Entry.Key;
++				var BuildTarget = Entry.Value.Item1;
++				var bBuildByDefault = Entry.Value.Item2;
++
++				try
++				{
++					BuildTarget.PreBuildSetup();
++
++					ExportedTargetInfo TargetInfo = ExportTarget(BuildTarget, bBuildByDefault, PlatformProjectGenerators);
++
++					DirectoryReference.CreateDirectory(OutputFile.Directory);
++
++					JsonSerializeOptions Options = Minimize == JsonWriterStyle.Readable 
++						? JsonSerializeOptions.PrettyPrint : JsonSerializeOptions.None;
++					File.WriteAllText(OutputFile.FullName, Json.Serialize(TargetInfo, Options));
++
++					ExportedTargetProjects.Add(TargetInfo);
++				}
++				catch (Exception Ex)
++				{
++					Log.TraceWarning("Exception while generating include data for Target:{Target}, Platform: {Platform}, Configuration: {Configuration}",
++						BuildTarget.AppName, BuildTarget.Platform.ToString(), BuildTarget.Configuration.ToString());
++					Log.TraceWarning("{Ex}", Ex.ToString());
++				}
++			}
++
++			return true;
++		}
++
++		private ExportedTargetInfo ExportTarget(
++			UEBuildTarget Target, bool bBuildByDefault, PlatformProjectGeneratorCollection PlatformProjectGenerators)
++		{
++			ExportedTargetInfo TargetInfo = new ExportedTargetInfo()
++			{
++				TargetName = Target.TargetName,
++				TargetPath = Target.TargetRulesFile.FullName,
++				ProjectPath = Target.ProjectFile?.FullName ?? String.Empty,
++				TargetType = Target.TargetType.ToString(),
++				Platform = Target.Platform.ToString(),
++				Configuration = Target.Configuration.ToString(),
++				BuildInfo = ExportBuildInfo(Target, PlatformProjectGenerators, bBuildByDefault)
++			};
++
++			UEToolChain TargetToolChain = Target.CreateToolchain(Target.Platform);
++			CppCompileEnvironment GlobalCompileEnvironment = Target.CreateCompileEnvironmentForProjectFiles();
++
++			HashSet<string> ModuleNames = new HashSet<string>();
++			foreach (UEBuildBinary Binary in Target.Binaries)
++			{
++				CppCompileEnvironment BinaryCompileEnvironment = Binary.CreateBinaryCompileEnvironment(GlobalCompileEnvironment);
++				IEnumerable<UEBuildModuleCPP> CandidateModules = Binary.Modules.Where(x => x is UEBuildModuleCPP).Cast<UEBuildModuleCPP>();
++
++				foreach (var ModuleCpp in CandidateModules)
++				{
++					if (!ModuleNames.Add(ModuleCpp.Name))
++					{
++						continue;
++					}
++
++					CppCompileEnvironment ModuleCompileEnvironment = ModuleCpp.CreateCompileEnvironmentForIntellisense(Target.Rules, BinaryCompileEnvironment);
++					TargetInfo.ModuleToCompileSettings.Add(ModuleCpp.Name, ExportModule(ModuleCpp, TargetToolChain, ModuleCompileEnvironment));
++
++					foreach (DirectoryReference ModuleDirectory in ModuleCpp.ModuleDirectories)
++					{
++						if (!TargetInfo.DirToModule.ContainsKey(ModuleDirectory.FullName))
++						{
++							TargetInfo.DirToModule.Add(ModuleDirectory.FullName, ModuleCpp.Name);
++						}
++					}
++
++					if (ModuleCpp.GeneratedCodeDirectory != null)
++					{
++                        if (!TargetInfo.DirToModule.ContainsKey(ModuleCpp.GeneratedCodeDirectory.FullName))
++                        {
++                            TargetInfo.DirToModule.Add(ModuleCpp.GeneratedCodeDirectory.FullName, ModuleCpp.Name);
++                        }
++					}
++				}
++			}
++
++			return TargetInfo;
++		}
++
++		private static ExportedModuleInfo ExportModule(UEBuildModuleCPP Module, UEToolChain TargetToolChain, CppCompileEnvironment ModuleCompileEnvironment)
++		{
++			ExportedModuleInfo Result = new ExportedModuleInfo()
++			{
++				Name = Module.Name,
++				Directory = Module.ModuleDirectory.FullName,
++				Rules = Module.RulesFile.FullName,
++				GeneratedCodeDirectory = Module.GeneratedCodeDirectory != null ? Module.GeneratedCodeDirectory.FullName : String.Empty,
++				Standard = ModuleCompileEnvironment.CppStandard.ToString(),
++			};
++
++			Result.IncludePaths.AddRange(Module.PublicIncludePaths.Select(x => x.FullName));
++			Result.IncludePaths.AddRange(Module.PublicSystemIncludePaths.Select(x => x.FullName));
++			Result.IncludePaths.AddRange(Module.LegacyPublicIncludePaths.Select(x => x.FullName));
++			Result.IncludePaths.AddRange(Module.PrivateIncludePaths.Select(x => x.FullName));
++
++			Result.IncludePaths.AddRange(ModuleCompileEnvironment.UserIncludePaths.Select(x => x.FullName));
++
++			Result.IncludePaths.AddRange(Module.PublicSystemLibraryPaths.Select(x => x.FullName));
++			Result.IncludePaths.AddRange(Module.PublicSystemLibraries.Concat(Module.PublicLibraries.Select(x => x.FullName)));
++
++			VCToolChain TargetVCToolChain = TargetToolChain as VCToolChain;
++			if (TargetToolChain != null)
++			{
++				string VCIncludePaths = VCToolChain.GetVCIncludePaths(ModuleCompileEnvironment.Platform, WindowsCompiler.VisualStudio2022, null);
++				Result.IncludePaths.AddRange(VCIncludePaths.Split(';'));
++			}
++
++			Result.Defines.AddRange(Module.PublicDefinitions);
++			Result.Defines.AddRange(Module.Rules.PrivateDefinitions);
++			Result.Defines.AddRange(Module.Rules.bTreatAsEngineModule ? Array.Empty<string>() : Module.Rules.Target.ProjectDefinitions);
++			Result.Defines.AddRange(Module.GetEmptyApiMacros());
++
++			var ForcedIncludes = ModuleCompileEnvironment.ForceIncludeFiles.ToList();
++			if (ModuleCompileEnvironment.PrecompiledHeaderAction == PrecompiledHeaderAction.Include)
++			{
++				FileItem IncludeHeader = FileItem.GetItemByFileReference(ModuleCompileEnvironment.PrecompiledHeaderIncludeFilename);
++				ForcedIncludes.Insert(0, IncludeHeader);
++			}
++
++			Result.ForcedIncludes.AddRange(ForcedIncludes.Select(x => x.FullName));
++
++			Result.CompilerPath = TargetToolChain.GetCppCompilerPath()?.ToString();
++			Result.CompilerArgs.AddRange(TargetToolChain.GetGlobalCommandLineArgs(ModuleCompileEnvironment));
++			Result.CompilerAdditionalArgs.Add("c", TargetToolChain.GetCCommandLineArgs(ModuleCompileEnvironment).ToList());
++			Result.CompilerAdditionalArgs.Add("cpp", TargetToolChain.GetCPPCommandLineArgs(ModuleCompileEnvironment).ToList());
++
++			return Result;
++		}
++
++		private TargetBuildInfo ExportBuildInfo(UEBuildTarget Target, PlatformProjectGeneratorCollection PlatformProjectGenerators, bool bBuildByDefault)
++		{
++			if (!BuildHostPlatform.Current.Platform.IsInGroup(UnrealPlatformGroup.Windows))
++			{
++				Log.TraceWarning("Unsupported platform for Build Information: {Platform}", BuildHostPlatform.Current.Platform.ToString());
++				return null;
++			}
++
++			if (IsStubProject)
++			{
++				return null;
++			}
++
++			ProjectTarget ProjectTarget = ProjectTargets.OfType<ProjectTarget>().Single(It => Target.TargetRulesFile == It.TargetFilePath);
++			UnrealTargetPlatform Platform = Target.Platform;
++			UnrealTargetConfiguration Configuration = Target.Configuration;
++
++			string UProjectPath = "";
++			if (IsForeignProject)
++			{
++				UProjectPath = String.Format("\"{0}\"", ProjectTarget.UnrealProjectFilePath.FullName);
++			}
++
++			PlatformProjectGenerator ProjGenerator = PlatformProjectGenerators.GetPlatformProjectGenerator(Platform, true);
++
++			string BuildArguments = GetBuildArguments(ProjGenerator, ProjectTarget, Target, UProjectPath);
++
++			return new TargetBuildInfo()
++			{
++				BuildCmd = $"{GetBuildScript("Build")} {BuildArguments}",
++				RebuildCmd = $"{GetBuildScript("Rebuild")} {BuildArguments}",
++				CleanCmd = $"{GetBuildScript("Clean")} {BuildArguments}",
++				PrimaryOutput = Target.Binaries[0].OutputFilePath.FullName,
++				BuildByDefault = bBuildByDefault,
++			};
++		}
++
++		private string GetBuildScript(string Cmd)
++		{
++			DirectoryReference BatchFilesDirectory = DirectoryReference.Combine(UnrealBuildTool.EngineDirectory, "Build", "BatchFiles");
++			return EscapePath(FileReference.Combine(BatchFilesDirectory, $"{Cmd}.bat").FullName);
++		}
++
++		private string GetBuildArguments(
++			PlatformProjectGenerator ProjGenerator, ProjectTarget ProjectTarget, UEBuildTarget Target, string UProjectPath)
++		{
++			UnrealTargetConfiguration Configuration = Target.Configuration;
++			UnrealTargetPlatform Platform = Target.Platform;
++			string TargetName = ProjectTarget.TargetFilePath.GetFileNameWithoutAnyExtensions();
++			TargetRules TargetRulesObject = ProjectTarget.TargetRules;
++
++			// This is the standard UE based project NMake build line:
++			//	..\..\Build\BatchFiles\Build.bat <TARGETNAME> <PLATFORM> <CONFIGURATION>
++			//	ie ..\..\Build\BatchFiles\Build.bat BlankProgram Win64 Debug
++
++			StringBuilder BuildArguments = new StringBuilder();
++
++			BuildArguments.AppendFormat("{0} {1} {2}", TargetName, Configuration, Platform);
++			if (IsForeignProject)
++			{
++				BuildArguments.AppendFormat(" -Project={0}", UProjectPath);
++			}
++
++			List<string> ExtraTargets = new List<string>();
++			if (!bUsePrecompiled)
++			{
++				if (TargetRulesObject.Type == TargetType.Editor && Settings.bEditorDependsOnShaderCompileWorker && !UnrealBuildTool.IsEngineInstalled())
++				{
++					ExtraTargets.Add("ShaderCompileWorker Win64 Development");
++				}
++				if (TargetRulesObject.bWithLiveCoding && Settings.bBuildLiveCodingConsole && !UnrealBuildTool.IsEngineInstalled() && TargetRulesObject.Name != "LiveCodingConsole")
++				{
++					ExtraTargets.Add(TargetRulesObject.bUseDebugLiveCodingConsole ? "LiveCodingConsole Win64 Debug" : "LiveCodingConsole Win64 Development");
++				}
++			}
++
++			if (ExtraTargets.Count > 0)
++			{
++				BuildArguments.Replace("\"", "\\\"");
++				BuildArguments.Insert(0, "-Target=\"");
++				BuildArguments.Append("\"");
++				foreach (string ExtraTarget in ExtraTargets)
++				{
++					BuildArguments.AppendFormat(" -Target=\"{0} -Quiet\"", ExtraTarget);
++				}
++			}
++
++			if (bUsePrecompiled)
++			{
++				BuildArguments.Append(" -UsePrecompiled");
++			}
++
++			// Always wait for the mutex between UBT invocations, so that building the whole solution doesn't fail.
++			BuildArguments.Append(" -WaitMutex");
++
++			// Always include a flag to format log messages for MSBuild
++			BuildArguments.Append(" -FromMsBuild");
++
++			if (Settings.bAddFastPDBToProjects)
++			{
++				// Pass Fast PDB option to make use of Visual Studio's /DEBUG:FASTLINK option
++				BuildArguments.Append(" -FastPDB");
++			}
++
++			if (ProjGenerator != null)
++			{
++				BuildArguments.Append(ProjGenerator.GetExtraBuildArguments(Platform, Configuration));
++			}
++
++			return BuildArguments.ToString();
++		}
++
++		internal class ExportedTargetInfo
++		{
++			public string TargetName { get; set; } = String.Empty;
++			public string TargetPath { get; set; } = String.Empty;
++			public string ProjectPath { get; set; } = String.Empty;
++			public string TargetType { get; set; } = String.Empty;
++			public string Platform { get; set; } = String.Empty;
++			public string Configuration { get; set; } = String.Empty;
++			public TargetBuildInfo BuildInfo { get; set; }
++			public Dictionary<string, ExportedModuleInfo> ModuleToCompileSettings { get; set; } = new Dictionary<string, ExportedModuleInfo>();
++			public Dictionary<string, string> DirToModule { get; set; } = new Dictionary<string, string>();
++		}
++
++		internal class ExportedModuleInfo
++		{
++			public string Name { get; set; } = String.Empty;
++			public string Directory { get; set; } = String.Empty;
++			public string Rules { get; set; } = String.Empty;
++			public string GeneratedCodeDirectory { get; set; } = String.Empty;
++			public List<string> IncludePaths { get; set; } = new List<string>();
++			public List<string> Defines { get; set; } = new List<string>();
++			public string Standard { get; set; }
++			public List<string> ForcedIncludes { get; set; } = new List<string>();
++			public string CompilerPath { get; set; }
++			public List<string> CompilerArgs { get; set; } = new List<string>();
++			public Dictionary<string, List<string>> CompilerAdditionalArgs { get; set; } = new Dictionary<string, List<string>>();
++			public string WindowsSdkVersion { get; set; }
++		}
++
++		internal class TargetBuildInfo
++		{
++			public string BuildCmd { get; set; } = String.Empty;
++			public string RebuildCmd { get; set; } = String.Empty;
++			public string CleanCmd { get; set; } = String.Empty;
++			public string PrimaryOutput { get; set; } = String.Empty;
++			public bool BuildByDefault { get; internal set; }
++		}
++	}
++}
+diff --git a/Engine/Source/Programs/UnrealBuildTool/ProjectFiles/VisualStudioWorkspace/VSWorkspaceProjectFileGenerator.cs b/Engine/Source/Programs/UnrealBuildTool/ProjectFiles/VisualStudioWorkspace/VSWorkspaceProjectFileGenerator.cs
+new file mode 100644
+index 0000000000000..df430da977f9e
+--- /dev/null
++++ b/Engine/Source/Programs/UnrealBuildTool/ProjectFiles/VisualStudioWorkspace/VSWorkspaceProjectFileGenerator.cs
+@@ -0,0 +1,277 @@
++// Copyright Epic Games, Inc. All Rights Reserved.
++
++using System;
++using System.Collections.Generic;
++using System.IO;
++using System.Linq;
++using System.Text.RegularExpressions;
++using Tools.DotNETCommon;
++
++namespace UnrealBuildTool
++{
++	class VSWorkspaceMasterProjectFolder : MasterProjectFolder
++	{
++		public VSWorkspaceMasterProjectFolder(ProjectFileGenerator InitOwnerProjectFileGenerator, string InitFolderName) :
++			base(InitOwnerProjectFileGenerator, InitFolderName)
++		{
++		}
++	}
++
++	class VSWorkspaceProjectFileGenerator : ProjectFileGenerator
++	{
++		public override string ProjectFileExtension => ".json";
++
++		// These properties are used by Visual Studio to determine where to read the project files.
++		// So they must remain constant.
++		private const string VSUnrealWorkspaceFileName = ".vs-unreal-workspace";
++		private const string ProjectFilesFolder = "VisualStudio";
++
++		private readonly CommandLineArguments Arguments;
++
++		/// <summary>
++		/// List of deprecated platforms.
++		/// Don't generate project model for these platforms unless they are specified in "Platforms" console arguments. 
++		/// </summary>
++		/// <returns></returns>
++		private readonly HashSet<UnrealTargetPlatform> DeprecatedPlatforms = new HashSet<UnrealTargetPlatform>();
++
++		/// <summary>
++		/// Platforms to generate project files for
++		/// </summary>
++		[CommandLine("-Platforms=", ListSeparator = '+')]
++		HashSet<UnrealTargetPlatform> Platforms = new HashSet<UnrealTargetPlatform>();
++
++		/// <summary>
++		/// Target types to generate project files for
++		/// </summary>
++		[CommandLine("-TargetTypes=", ListSeparator = '+')]
++		HashSet<TargetType> TargetTypes = new HashSet<TargetType>();
++
++		/// <summary>
++		/// Target configurations to generate project files for
++		/// </summary>
++		[CommandLine("-TargetConfigurations=", ListSeparator = '+')]
++		HashSet<UnrealTargetConfiguration> TargetConfigurations = new HashSet<UnrealTargetConfiguration>();
++
++		/// <summary>
++		/// Projects to generate project files for
++		/// </summary>
++		[CommandLine("-ProjectNames=", ListSeparator = '+')]
++		HashSet<string> ProjectNames = new HashSet<string>();
++
++		/// <summary>
++		/// Should format JSON files in human readable form, or use packed one without indents
++		/// </summary>
++		[CommandLine("-Minimize", Value = "Compact")]
++		private JsonWriterStyle Minimize = JsonWriterStyle.Readable;
++
++		/// <summary>
++		/// The settings object
++		/// </summary>
++		protected VCProjectFileSettings Settings = new VCProjectFileSettings();
++
++		public VSWorkspaceProjectFileGenerator(FileReference InOnlyGameProject,
++			CommandLineArguments InArguments)
++			: base(InOnlyGameProject)
++		{
++			Arguments = InArguments;
++			Arguments.ApplyTo(this);
++			XmlConfig.ApplyTo(Settings);
++		}
++
++		public override bool ShouldGenerateIntelliSenseData() => true;
++
++		public override void CleanProjectFiles(DirectoryReference InPrimaryProjectDirectory, string InPrimaryProjectName,
++			DirectoryReference InIntermediateProjectFilesDirectory)
++		{
++			DirectoryReference.Delete(InPrimaryProjectDirectory);
++		}
++
++		protected override void ConfigureProjectFileGeneration(string[] Arguments, ref bool IncludeAllPlatforms)
++		{
++			base.ConfigureProjectFileGeneration(Arguments, ref IncludeAllPlatforms);
++		}
++
++		protected override ProjectFile AllocateProjectFile(FileReference InitFilePath)
++		{
++			VSWorkspaceProjectFile projectFile = new VSWorkspaceProjectFile(InitFilePath, InitFilePath.Directory,
++				TargetTypes, Arguments, Settings, bUsePrecompiled);
++			return projectFile;
++		}
++
++		public override MasterProjectFolder AllocateMasterProjectFolder(ProjectFileGenerator OwnerProjectFileGenerator,
++			string FolderName)
++		{
++			return new VSWorkspaceMasterProjectFolder(OwnerProjectFileGenerator, FolderName);
++		}
++
++		protected override bool WriteProjectFiles(PlatformProjectGeneratorCollection PlatformProjectGenerators)
++		{
++			using (ProgressWriter Progress = new ProgressWriter("Writing project files...", true))
++			{
++				List<ProjectFile> ProjectsToGenerate = new List<ProjectFile>(GeneratedProjectFiles);
++				if (ProjectNames.Any())
++				{
++					ProjectsToGenerate = ProjectsToGenerate.Where(Project =>
++						ProjectNames.Contains(Project.ProjectFilePath.GetFileNameWithoutAnyExtensions())).ToList();
++				}
++
++				int TotalProjectFileCount = ProjectsToGenerate.Count;
++
++				HashSet<UnrealTargetPlatform> PlatformsToGenerate = new HashSet<UnrealTargetPlatform>(SupportedPlatforms);
++				if (Platforms.Any())
++				{
++					PlatformsToGenerate.IntersectWith(Platforms);
++				}
++
++				List<UnrealTargetPlatform> FilteredPlatforms = PlatformsToGenerate.Where(Platform =>
++				{
++					// Skip deprecated unless explicitly specified in the command line.
++					return (!DeprecatedPlatforms.Contains(Platform) || Platforms.Contains(Platform))
++						   && UEBuildPlatform.IsPlatformAvailable(Platform)
++						   && IsPlatformInHostGroup(Platform);
++				}).ToList();
++
++				HashSet<UnrealTargetConfiguration> ConfigurationsToGenerate = new HashSet<UnrealTargetConfiguration>(SupportedConfigurations);
++				if (TargetConfigurations.Any())
++				{
++					ConfigurationsToGenerate.IntersectWith(TargetConfigurations);
++				}
++
++				List<UnrealTargetConfiguration> Configurations = ConfigurationsToGenerate.ToList();
++
++				for (int ProjectFileIndex = 0; ProjectFileIndex < ProjectsToGenerate.Count; ++ProjectFileIndex)
++				{
++					VSWorkspaceProjectFile CurrentProject = ProjectsToGenerate[ProjectFileIndex] as VSWorkspaceProjectFile;
++					if (CurrentProject == null)
++					{
++						return false;
++					}
++
++					if (!CurrentProject.WriteProjectFile(FilteredPlatforms, Configurations, PlatformProjectGenerators, Minimize))
++					{
++						return false;
++					}
++
++					Progress.Write(ProjectFileIndex + 1, TotalProjectFileCount);
++				}
++
++				Progress.Write(TotalProjectFileCount, TotalProjectFileCount);
++			}
++			
++			return true;
++		}
++
++		public override bool GenerateProjectFiles(PlatformProjectGeneratorCollection PlatformProjectGenerators,
++			String[] Arguments)
++		{
++			bool IncludeAllPlatforms = true;
++			ConfigureProjectFileGeneration(Arguments, ref IncludeAllPlatforms);
++
++			if (bGeneratingGameProjectFiles || UnrealBuildTool.IsEngineInstalled())
++			{
++				MasterProjectPath = OnlyGameProject.Directory;
++				MasterProjectName = OnlyGameProject.GetFileNameWithoutExtension();
++
++				IntermediateProjectFilesPath =
++					DirectoryReference.Combine(MasterProjectPath, "Intermediate", "ProjectFiles");
++			}
++
++			string SupportedPlatformNames;
++			SetupSupportedPlatformsAndConfigurations(IncludeAllPlatforms: true, SupportedPlatformNames: out SupportedPlatformNames);
++			Log.TraceVerbose("Supported platforms: {0}", SupportedPlatformNames);
++
++			List<FileReference> AllGames = FindGameProjects();
++
++			{
++				// Find all of the target files.
++				List<FileReference> AllTargetFiles = DiscoverTargets(AllGames);
++
++				// If there are multiple targets of a given type for a project, use the order to determine which one gets a suffix.
++				AllTargetFiles = AllTargetFiles.OrderBy(x => x.FullName, StringComparer.OrdinalIgnoreCase).ToList();
++
++				ProjectFile EngineProject = null;
++				ProjectFile EnterpriseProject = null;
++				List<ProjectFile> GameProjects; 
++				Dictionary<FileReference, ProjectFile> ProgramProjects;
++
++				AddProjectsForAllTargets(
++					PlatformProjectGenerators,
++					AllGames,
++					AllTargetFiles,
++					Arguments,
++					out EngineProject,
++					out EnterpriseProject,
++					out GameProjects,
++					out ProgramProjects);
++
++				AddAllGameProjects(GameProjects, SupportedPlatformNames, ProjectsFolder: null);
++			}
++
++			WriteProjectFiles(PlatformProjectGenerators);
++			WriteMasterProjectFile(UBTProject, PlatformProjectGenerators);
++
++			return true;
++		}
++
++		protected override bool WriteMasterProjectFile(ProjectFile UBTProject,
++			PlatformProjectGeneratorCollection PlatformProjectGenerators)
++		{
++			try
++			{
++				FileReference PrimaryProjectFile = FileReference.Combine(
++					IntermediateProjectFilesPath, ProjectFilesFolder, VSUnrealWorkspaceFileName);
++				DirectoryReference.CreateDirectory(PrimaryProjectFile.Directory);
++
++				// Collect all the resulting project files and aggregate the target-level data
++				var AggregatedProjectInfo = GeneratedProjectFiles
++					.Where(Project => Project is VSWorkspaceProjectFile)
++					.OfType<VSWorkspaceProjectFile>()
++					.SelectMany(Project => Project.ExportedTargetProjects)
++					.GroupBy(TargetProject => TargetProject.TargetName)
++					.Select(g => new
++					{
++						g.Key,
++						Target = new
++						{
++							TargetType = g.Select(i => i.TargetType).Distinct().Single(),
++							TargetPath = g.Select(i => i.TargetPath).Distinct().Single(),
++							ProjectPath = g.Select(i => i.ProjectPath).Distinct().Single(),
++							Configurations = g.Select(i => i.Configuration).Distinct().ToList(),
++							Platforms = g.Select(i => i.Platform).Distinct().ToList(),
++						}
++					});
++
++				// The inner Targets object is needed for schema compatibility with the Query Mode API.
++				var Result = new
++				{
++					Targets = AggregatedProjectInfo.ToDictionary(item => item.Key, item => item.Target)
++				};
++
++				JsonSerializeOptions Options = Minimize == JsonWriterStyle.Readable
++						? JsonSerializeOptions.PrettyPrint : JsonSerializeOptions.None;
++				File.WriteAllText(PrimaryProjectFile.FullName, Json.Serialize(Result, Options));
++			}
++			catch (Exception Ex)
++			{
++				Log.TraceWarning("Exception while writing root project file: {0}", Ex.ToString());
++				return false;
++			}
++
++			return true;
++		}
++
++		/// <inheritdoc />
++		protected override FileReference GetProjectLocation(string BaseName)
++		{
++			return FileReference.Combine(IntermediateProjectFilesPath, ProjectFilesFolder, BaseName + ProjectFileExtension);
++		}
++
++		private static bool IsPlatformInHostGroup(UnrealTargetPlatform Platform)
++		{
++			return UEBuildPlatform.GetPlatformGroups(BuildHostPlatform.Current.Platform)
++				.Where(Group => Group != UnrealPlatformGroup.Desktop)
++				.Any(Group => UEBuildPlatform.IsPlatformInGroup(Platform, Group));
++		}
++	}
++}
+diff --git a/Engine/Source/Programs/UnrealBuildTool/ToolChain/UEToolChain.cs b/Engine/Source/Programs/UnrealBuildTool/ToolChain/UEToolChain.cs
+index 731ca46cf79e2..7a3b6c2825883 100644
+--- a/Engine/Source/Programs/UnrealBuildTool/ToolChain/UEToolChain.cs
++++ b/Engine/Source/Programs/UnrealBuildTool/ToolChain/UEToolChain.cs
+@@ -27,6 +27,17 @@ public virtual void GetExternalDependencies(HashSet<FileItem> ExternalDependenci
+ 		{
+ 		}
+ 
++		public static DirectoryReference GetModuleInterfaceDir(DirectoryReference OutputDir)
++		{
++			return DirectoryReference.Combine(OutputDir, "Ifc");
++		}
++
++		// Return the path to the cpp compiler that will be used by this toolchain.
++		public virtual FileReference GetCppCompilerPath()
++		{
++			return null;
++		}
++
+ 		public abstract CPPOutput CompileCPPFiles(CppCompileEnvironment CompileEnvironment, List<FileItem> InputFiles, DirectoryReference OutputDir, string ModuleName, IActionGraphBuilder Graph);
+ 
+ 		public virtual CPPOutput CompileRCFiles(CppCompileEnvironment Environment, List<FileItem> InputFiles, DirectoryReference OutputDir, IActionGraphBuilder Graph)
+@@ -57,6 +68,20 @@ public virtual FileItem[] LinkAllFiles(LinkEnvironment LinkEnvironment, bool bBu
+ 			return new FileItem[] { LinkFiles(LinkEnvironment, bBuildImportLibraryOnly, Graph) };
+ 		}
+ 
++		public virtual IEnumerable<string> GetGlobalCommandLineArgs(CppCompileEnvironment CompileEnvironment)
++		{
++			return Array.Empty<string>();
++		}
++
++		public virtual IEnumerable<string> GetCPPCommandLineArgs(CppCompileEnvironment CompileEnvironment)
++		{
++			return Array.Empty<string>();
++		}
++
++		public virtual IEnumerable<string> GetCCommandLineArgs(CppCompileEnvironment CompileEnvironment)
++		{
++			return Array.Empty<string>();
++		}
+ 
+ 		/// <summary>
+ 		/// Get the name of the response file for the current linker environment and output file
+diff --git a/Engine/Source/Programs/UnrealBuildTool/UnrealBuildTool.csproj b/Engine/Source/Programs/UnrealBuildTool/UnrealBuildTool.csproj
+index 3b070a635e74e..51818e6fa3a08 100644
+--- a/Engine/Source/Programs/UnrealBuildTool/UnrealBuildTool.csproj
++++ b/Engine/Source/Programs/UnrealBuildTool/UnrealBuildTool.csproj
+@@ -193,6 +193,8 @@
+     <Compile Include="ProjectFiles\VisualStudio\VCProject.cs" />
+     <Compile Include="ProjectFiles\VisualStudio\VCProjectFileGenerator.cs" />
+     <Compile Include="ProjectFiles\VisualStudio\VCSolutionOptions.cs" />
++    <Compile Include="ProjectFiles\VisualStudioWorkspace\VSWorkspaceProjectFile.cs" />
++    <Compile Include="ProjectFiles\VisualStudioWorkspace\VSWorkspaceProjectFileGenerator.cs" />
+     <Compile Include="ProjectFiles\Xcode\XcodeFrameworkWrapperProject.cs" />
+     <Compile Include="ProjectFiles\Xcode\XcodeProject.cs" />
+     <Compile Include="ProjectFiles\Xcode\XcodeProjectFileGenerator.cs" />
diff --git a/Engine/Extras/VisualStudioWorkspace/UnrealBuildTool-5.0.patch b/Engine/Extras/VisualStudioWorkspace/UnrealBuildTool-5.0.patch
new file mode 100644
--- /dev/null
+++ b/Engine/Extras/VisualStudioWorkspace/UnrealBuildTool-5.0.patch
@@ -0,0 +1,828 @@
+From 53e7e3f327d27608575225420cd23829acfcd69b Mon Sep 17 00:00:00 2001
+From: [REDACTED_PERSON] <[REDACTED_EMAIL]>
+Date: Wed, 17 Jan 2024 19:49:54 -0800
+Subject: [PATCH] [UE5.0] Project Generator for VS Workspace support
+
+Adding a new Project Generator for Visual Studio that will be used for
+the new support to opening `.uproject` files directly in the IDE without
+generating the traditional `.sln` and .vcxproj files.
+
+This is based on the QueryMode API in `ue5-main` and adapted to the
+GenerateProjectFiles flow. The project files will contain information for each
+Target+Configuration+Platform combinations that are used the project.
+
+Visual Studio will automatically invoke this new generator during the flow to
+open a `.uproject` file directly.
+---
+ .../Configuration/UEBuildTarget.cs            |   2 +-
+ .../Modes/GenerateProjectFilesMode.cs         |   4 +
+ .../Platform/Windows/VCToolChain.cs           |  26 ++
+ .../ProjectFiles/ProjectFileGenerator.cs      |   5 +-
+ .../VSWorkspaceProjectFile.cs                 | 394 ++++++++++++++++++
+ .../VSWorkspaceProjectFileGenerator.cs        | 248 +++++++++++
+ .../UnrealBuildTool/ToolChain/UEToolChain.cs  |  20 +
+ 7 files changed, 696 insertions(+), 3 deletions(-)
+ create mode 100644 Engine/Source/Programs/UnrealBuildTool/ProjectFiles/VisualStudioWorkspace/VSWorkspaceProjectFile.cs
+ create mode 100644 Engine/Source/Programs/UnrealBuildTool/ProjectFiles/VisualStudioWorkspace/VSWorkspaceProjectFileGenerator.cs
+
+diff --git a/Engine/Source/Programs/UnrealBuildTool/Configuration/UEBuildTarget.cs b/Engine/Source/Programs/UnrealBuildTool/Configuration/UEBuildTarget.cs
+index 25e98e1729cfb..04395105aaf64 100644
+--- a/Engine/Source/Programs/UnrealBuildTool/Configuration/UEBuildTarget.cs
++++ b/Engine/Source/Programs/UnrealBuildTool/Configuration/UEBuildTarget.cs
+@@ -2249,7 +2249,7 @@ List<RestrictedFolder> GetRestrictedFolders(FileReference File)
+ 		/// Creates a toolchain for the current target. May be overridden by the target rules.
+ 		/// </summary>
+ 		/// <returns>New toolchain instance</returns>
+-		private UEToolChain CreateToolchain(UnrealTargetPlatform Platform)
++		public UEToolChain CreateToolchain(UnrealTargetPlatform Platform)
+ 		{
+ 			if (Rules.ToolChainName == null)
+ 			{
+diff --git a/Engine/Source/Programs/UnrealBuildTool/Modes/GenerateProjectFilesMode.cs b/Engine/Source/Programs/UnrealBuildTool/Modes/GenerateProjectFilesMode.cs
+index 5ac8f51d12f4f..a6e0cde49cfa2 100644
+--- a/Engine/Source/Programs/UnrealBuildTool/Modes/GenerateProjectFilesMode.cs
++++ b/Engine/Source/Programs/UnrealBuildTool/Modes/GenerateProjectFilesMode.cs
+@@ -35,6 +35,7 @@ class GenerateProjectFilesMode : ToolMode
+ 		[CommandLine("-EddieProjectFiles", Value = nameof(ProjectFileFormat.Eddie))]
+ 		[CommandLine("-VSCode", Value = nameof(ProjectFileFormat.VisualStudioCode))]
+ 		[CommandLine("-VSMac", Value = nameof(ProjectFileFormat.VisualStudioMac))]
++		[CommandLine("-VSWorkspace", Value = nameof(ProjectFileFormat.VisualStudioWorkspace))]
+ 		[CommandLine("-CLion", Value = nameof(ProjectFileFormat.CLion))]
+ 		[CommandLine("-Rider", Value = nameof(ProjectFileFormat.Rider))]
+ 		#if __VPROJECT_AVAILABLE__
+@@ -202,6 +203,9 @@ public override int Execute(CommandLineArguments Arguments)
+ 					case ProjectFileFormat.VisualStudioCode:
+ 						Generator = new VSCodeProjectFileGenerator(ProjectFile);
+ 						break;
++					case ProjectFileFormat.VisualStudioWorkspace:
++						Generator = new VSWorkspaceProjectFileGenerator(ProjectFile, Arguments);
++						break;
+ 					case ProjectFileFormat.CLion:
+ 						Generator = new CLionGenerator(ProjectFile);
+ 						break;
+diff --git a/Engine/Source/Programs/UnrealBuildTool/Platform/Windows/VCToolChain.cs b/Engine/Source/Programs/UnrealBuildTool/Platform/Windows/VCToolChain.cs
+index 80ec89d70e37a..287e8abd3e78f 100644
+--- a/Engine/Source/Programs/UnrealBuildTool/Platform/Windows/VCToolChain.cs
++++ b/Engine/Source/Programs/UnrealBuildTool/Platform/Windows/VCToolChain.cs
+@@ -47,6 +47,11 @@ public VCToolChain(ReadOnlyTargetRules Target)
+ 			}
+ 		}
+ 
++		public override FileReference? GetCppCompilerPath()
++		{
++			return EnvVars.CompilerPath;
++		}
++
+ 		/// <summary>
+ 		/// Prepares the environment for building
+ 		/// </summary>
+@@ -209,6 +214,27 @@ public static void AddSourceDependsFile(List<string> Arguments, FileItem SourceD
+ 			Arguments.Add($"/clang:-MD /clang:-MF\"{SourceDependsFileString}\"");
+ 		}
+ 
++		public override IEnumerable<string> GetGlobalCommandLineArgs(CppCompileEnvironment CompileEnvironment)
++		{
++			List<string> Arguments = new List<string>();
++			AppendCLArguments_Global(new CppCompileEnvironment(CompileEnvironment), Arguments);
++			return Arguments;
++		}
++
++		public override IEnumerable<string> GetCPPCommandLineArgs(CppCompileEnvironment CompileEnvironment)
++		{
++			List<string> Arguments = new List<string>();
++			AppendCLArguments_CPP(CompileEnvironment, Arguments);
++			return Arguments;
++		}
++
++		public override IEnumerable<string> GetCCommandLineArgs(CppCompileEnvironment CompileEnvironment)
++		{
++			List<string> Arguments = new List<string>();
++			AppendCLArguments_C(Arguments);
++			return Arguments;
++		}
++
+ 		protected virtual void AppendCLArguments_Global(CppCompileEnvironment CompileEnvironment, List<string> Arguments)
+ 		{
+ 			// Suppress generation of object code for unreferenced inline functions. Enabling this option is more standards compliant, and causes a big reduction
+diff --git a/Engine/Source/Programs/UnrealBuildTool/ProjectFiles/ProjectFileGenerator.cs b/Engine/Source/Programs/UnrealBuildTool/ProjectFiles/ProjectFileGenerator.cs
+index 1e8e7e609c166..4276988ed1374 100644
+--- a/Engine/Source/Programs/UnrealBuildTool/ProjectFiles/ProjectFileGenerator.cs
++++ b/Engine/Source/Programs/UnrealBuildTool/ProjectFiles/ProjectFileGenerator.cs
+@@ -132,6 +132,7 @@ enum ProjectFileFormat
+ 		VisualStudio,
+ 		VisualStudio2019,
+ 		VisualStudio2022,
++		VisualStudioWorkspace,
+ 		XCode,
+ 		Eddie,
+ 		VisualStudioCode,
+@@ -2297,7 +2298,7 @@ private ProjectFile FindProjectForModule(FileReference CurModuleFile, List<FileR
+ 		/// <param name="GameProjects">Map of game folder name to all of the game projects we created</param>
+ 		/// <param name="ProgramProjects">Map of program names to all of the program projects we created</param>
+ 		/// <param name="RulesAssemblies">Map of RuleAssemblies to their base folders</param>
+-		private void AddProjectsForAllTargets(
++		protected void AddProjectsForAllTargets(
+ 			PlatformProjectGeneratorCollection PlatformProjectGenerators,
+ 			List<FileReference> AllGames,
+ 			List<FileReference> AllTargetFiles,
+@@ -2619,7 +2620,7 @@ protected void AddProjectsForMods(List<ProjectFile> GameProjects, out List<Proje
+ 		/// </summary>
+ 		/// <param name="BaseName">The base name for the project file</param>
+ 		/// <returns>Full path to the project file</returns>
+-		protected FileReference GetProjectLocation(string BaseName)
++		protected virtual FileReference GetProjectLocation(string BaseName)
+ 		{
+ 			return FileReference.Combine(IntermediateProjectFilesPath, BaseName + ProjectFileExtension);
+ 		}
+diff --git a/Engine/Source/Programs/UnrealBuildTool/ProjectFiles/VisualStudioWorkspace/VSWorkspaceProjectFile.cs b/Engine/Source/Programs/UnrealBuildTool/ProjectFiles/VisualStudioWorkspace/VSWorkspaceProjectFile.cs
+new file mode 100644
+index 0000000000000..05c5473a616db
+--- /dev/null
++++ b/Engine/Source/Programs/UnrealBuildTool/ProjectFiles/VisualStudioWorkspace/VSWorkspaceProjectFile.cs
+@@ -0,0 +1,394 @@
++// Copyright Epic Games, Inc. All Rights Reserved.
++
++using System;
++using System.Collections.Generic;
++using System.IO;
++using System.Linq;
++using System.Text;
++using System.Text.Json;
++using EpicGames.Core;
++using Microsoft.Extensions.Logging;
++using UnrealBuildBase;
++
++namespace UnrealBuildTool
++{
++	internal class VSWorkspaceProjectFile : ProjectFile
++	{
++		private readonly DirectoryReference RootPath;
++		private readonly HashSet<TargetType> TargetTypes;
++		private readonly CommandLineArguments Arguments;
++		private VCProjectFileSettings Settings;
++		private bool bUsePrecompiled;
++
++		/// <summary>
++		/// Collection of output files for this project
++		/// </summary>
++		public List<ExportedTargetInfo> ExportedTargetProjects { get; set; } = new List<ExportedTargetInfo>();
++
++		public VSWorkspaceProjectFile(FileReference InProjectFilePath, DirectoryReference BaseDir,
++			DirectoryReference RootPath, HashSet<TargetType> TargetTypes, CommandLineArguments Arguments, VCProjectFileSettings Settings, bool bUsePrecompiled)
++			: base(InProjectFilePath, BaseDir)
++		{
++			this.RootPath = RootPath;
++			this.TargetTypes = TargetTypes;
++			this.Arguments = Arguments;
++			this.Settings = Settings;
++			this.bUsePrecompiled = bUsePrecompiled;
++		}
++
++		/// <summary>
++		/// Write project file info in JSON file.
++		/// For every combination of <c>UnrealTargetPlatform</c>, <c>UnrealTargetConfiguration</c> and <c>TargetType</c>
++		/// will be generated separate JSON file.
++		/// </summary>
++		public bool WriteProjectFile(List<UnrealTargetPlatform> InPlatforms,
++			List<UnrealTargetConfiguration> InConfigurations,
++			PlatformProjectGeneratorCollection PlatformProjectGenerators, JsonWriterStyle Minimize)
++		{
++			DirectoryReference ProjectRootFolder = RootPath;
++
++			Dictionary<FileReference, (UEBuildTarget BuildTarget, bool bBuildByDefault)> FileToTarget = new Dictionary<FileReference, (UEBuildTarget, bool)>();
++
++			foreach (UnrealTargetPlatform Platform in InPlatforms)
++			{
++				foreach (UnrealTargetConfiguration Configuration in InConfigurations)
++				{
++					foreach (ProjectTarget ProjectTarget in ProjectTargets.OfType<ProjectTarget>())
++					{
++						if (TargetTypes.Any() && !TargetTypes.Contains(ProjectTarget.TargetRules!.Type))
++						{
++							continue;
++						}
++
++						// Skip Programs for all configs except for current platform + Development & Debug configurations
++						if (ProjectTarget.TargetRules!.Type == TargetType.Program &&
++							(BuildHostPlatform.Current.Platform != Platform ||
++							 !(Configuration == UnrealTargetConfiguration.Development || Configuration == UnrealTargetConfiguration.Debug)))
++						{
++							continue;
++						}
++
++						// Skip Editor for all platforms except for current platform
++						if (ProjectTarget.TargetRules.Type == TargetType.Editor &&
++							(BuildHostPlatform.Current.Platform != Platform ||
++							(Configuration == UnrealTargetConfiguration.Test || Configuration == UnrealTargetConfiguration.Shipping)))
++						{
++							continue;
++						}
++
++						bool bBuildByDefault = ShouldBuildByDefaultForSolutionTargets && ProjectTarget.SupportedPlatforms.Contains(Platform);
++
++						string DefaultArchitecture = UEBuildPlatform
++							.GetBuildPlatform(Platform)
++							.GetDefaultArchitecture(ProjectTarget.UnrealProjectFilePath);
++
++						TargetDescriptor TargetDesc = new TargetDescriptor(ProjectTarget.UnrealProjectFilePath, ProjectTarget.Name,
++							Platform, Configuration, DefaultArchitecture, Arguments);
++
++						try
++						{
++							FileReference OutputFile = FileReference.Combine(ProjectRootFolder,
++								$"{ProjectTarget.TargetFilePath.GetFileNameWithoutAnyExtensions()}_{Configuration}_{Platform}.json");
++
++							UEBuildTarget BuildTarget = UEBuildTarget.Create(TargetDesc, false, false, false);
++							FileToTarget.Add(OutputFile, (BuildTarget, bBuildByDefault));
++						}
++						catch (Exception Ex)
++						{
++							Log.TraceWarning("Exception while generating include data for Target:{Target}, Platform: {Platform}, Configuration: {Configuration}", TargetDesc.Name, Platform.ToString(), Configuration.ToString());
++							Log.TraceWarning("{Ex}", Ex.ToString());
++						}
++					}
++				}
++			}
++
++			foreach (var Entry in FileToTarget)
++			{
++				var OutputFile = Entry.Key;
++				var BuildTarget = Entry.Value.BuildTarget;
++				var bBuildByDefault = Entry.Value.bBuildByDefault;
++
++				try
++				{
++					BuildTarget.PreBuildSetup();
++
++					ExportedTargetInfo TargetInfo = ExportTarget(BuildTarget, bBuildByDefault, PlatformProjectGenerators);
++
++					DirectoryReference.CreateDirectory(OutputFile.Directory);
++
++					File.WriteAllText(OutputFile.FullName, JsonSerializer.Serialize(TargetInfo, options: new JsonSerializerOptions()
++					{
++						PropertyNamingPolicy = JsonNamingPolicy.CamelCase,
++						WriteIndented = Minimize == JsonWriterStyle.Readable,
++					}));
++
++					ExportedTargetProjects.Add(TargetInfo);
++				}
++				catch (Exception Ex)
++				{
++					Log.TraceWarning("Exception while generating include data for Target:{Target}, Platform: {Platform}, Configuration: {Configuration}",
++						BuildTarget.AppName, BuildTarget.Platform.ToString(), BuildTarget.Configuration.ToString());
++					Log.TraceWarning("{Ex}", Ex.ToString());
++				}
++			}
++
++			return true;
++		}
++
++		private ExportedTargetInfo ExportTarget(
++			UEBuildTarget Target, bool bBuildByDefault, PlatformProjectGeneratorCollection PlatformProjectGenerators)
++		{
++			ExportedTargetInfo TargetInfo = new ExportedTargetInfo()
++			{
++				TargetName = Target.TargetName,
++				TargetPath = Target.TargetRulesFile.FullName,
++				ProjectPath = Target.ProjectFile?.FullName ?? String.Empty,
++				TargetType = Target.TargetType.ToString(),
++				Platform = Target.Platform.ToString(),
++				Configuration = Target.Configuration.ToString(),
++				BuildInfo = ExportBuildInfo(Target, PlatformProjectGenerators, bBuildByDefault)
++			};
++
++			UEToolChain TargetToolChain = Target.CreateToolchain(Target.Platform);
++			CppCompileEnvironment GlobalCompileEnvironment = Target.CreateCompileEnvironmentForProjectFiles();
++
++			HashSet<string> ModuleNames = new HashSet<string>();
++			foreach (UEBuildBinary Binary in Target.Binaries)
++			{
++				CppCompileEnvironment BinaryCompileEnvironment = Binary.CreateBinaryCompileEnvironment(GlobalCompileEnvironment);
++				IEnumerable<UEBuildModuleCPP> CandidateModules = Binary.Modules.Where(x => x is UEBuildModuleCPP).Cast<UEBuildModuleCPP>();
++
++				foreach (var ModuleCpp in CandidateModules)
++				{
++					if (!ModuleNames.Add(ModuleCpp.Name))
++					{
++						continue;
++					}
++
++					CppCompileEnvironment ModuleCompileEnvironment = ModuleCpp.CreateCompileEnvironmentForIntellisense(Target.Rules, BinaryCompileEnvironment);
++					TargetInfo.ModuleToCompileSettings.Add(ModuleCpp.Name, ExportModule(ModuleCpp, TargetToolChain, ModuleCompileEnvironment));
++
++					foreach (DirectoryReference ModuleDirectory in ModuleCpp.ModuleDirectories)
++					{
++						TargetInfo.DirToModule.TryAdd(ModuleDirectory.FullName, ModuleCpp.Name);
++					}
++
++					if (ModuleCpp.GeneratedCodeDirectory != null)
++					{
++						TargetInfo.DirToModule.TryAdd(ModuleCpp.GeneratedCodeDirectory.FullName, ModuleCpp.Name);
++					}
++				}
++			}
++
++			return TargetInfo;
++		}
++
++		private static ExportedModuleInfo ExportModule(UEBuildModuleCPP Module, UEToolChain TargetToolChain, CppCompileEnvironment ModuleCompileEnvironment)
++		{
++			ExportedModuleInfo Result = new ExportedModuleInfo()
++			{
++				Name = Module.Name,
++				Directory = Module.ModuleDirectory.FullName,
++				Rules = Module.RulesFile.FullName,
++				GeneratedCodeDirectory = Module.GeneratedCodeDirectory != null ? Module.GeneratedCodeDirectory.FullName : String.Empty,
++				Standard = ModuleCompileEnvironment.CppStandard.ToString(),
++			};
++
++			Result.IncludePaths.AddRange(Module.PublicIncludePaths.Select(x => x.FullName));
++			Result.IncludePaths.AddRange(Module.PublicSystemIncludePaths.Select(x => x.FullName));
++			Result.IncludePaths.AddRange(Module.InternalIncludePaths.Select(x => x.FullName));
++			Result.IncludePaths.AddRange(Module.LegacyPublicIncludePaths.Select(x => x.FullName));
++			Result.IncludePaths.AddRange(Module.PrivateIncludePaths.Select(x => x.FullName));
++
++			Result.IncludePaths.AddRange(ModuleCompileEnvironment.UserIncludePaths.Select(x => x.FullName));
++
++			Result.IncludePaths.AddRange(Module.PublicSystemLibraryPaths.Select(x => x.FullName));
++			Result.IncludePaths.AddRange(Module.PublicSystemLibraries.Concat(Module.PublicLibraries.Select(x => x.FullName)));
++
++			if (TargetToolChain is VCToolChain TargetVCToolChain)
++			{
++				string VCIncludePaths = VCToolChain.GetVCIncludePaths(ModuleCompileEnvironment.Platform, WindowsCompiler.VisualStudio2022, null);
++				Result.IncludePaths.AddRange(VCIncludePaths.Split(";"));
++			}
++
++			Result.Defines.AddRange(Module.PublicDefinitions);
++			Result.Defines.AddRange(Module.Rules.PrivateDefinitions);
++			Result.Defines.AddRange(Module.Rules.bTreatAsEngineModule ? Array.Empty<string>() : Module.Rules.Target.ProjectDefinitions);
++			Result.Defines.AddRange(Module.GetEmptyApiMacros());
++
++			var ForcedIncludes = ModuleCompileEnvironment.ForceIncludeFiles.ToList();
++			if (ModuleCompileEnvironment.PrecompiledHeaderAction == PrecompiledHeaderAction.Include)
++			{
++				FileItem IncludeHeader = FileItem.GetItemByFileReference(ModuleCompileEnvironment.PrecompiledHeaderIncludeFilename!);
++				ForcedIncludes.Insert(0, IncludeHeader);
++			}
++
++			Result.ForcedIncludes.AddRange(ForcedIncludes.Select(x => x.FullName));
++
++			Result.CompilerPath = TargetToolChain.GetCppCompilerPath()?.ToString();
++			Result.CompilerArgs.AddRange(TargetToolChain.GetGlobalCommandLineArgs(ModuleCompileEnvironment));
++			Result.CompilerAdditionalArgs.Add("c", TargetToolChain.GetCCommandLineArgs(ModuleCompileEnvironment).ToList());
++			Result.CompilerAdditionalArgs.Add("cpp", TargetToolChain.GetCPPCommandLineArgs(ModuleCompileEnvironment).ToList());
++
++			return Result;
++		}
++
++		private TargetBuildInfo? ExportBuildInfo(UEBuildTarget Target, PlatformProjectGeneratorCollection PlatformProjectGenerators, bool bBuildByDefault)
++		{
++			if (!BuildHostPlatform.Current.Platform.IsInGroup(UnrealPlatformGroup.Windows))
++			{
++				Log.TraceWarning("Unsupported platform for Build Information: {Platform}", BuildHostPlatform.Current.Platform.ToString());
++				return null;
++			}
++
++			if (IsStubProject)
++			{
++				return null;
++			}
++
++			ProjectTarget ProjectTarget = ProjectTargets.OfType<ProjectTarget>().Single(It => Target.TargetRulesFile == It.TargetFilePath);
++			UnrealTargetPlatform Platform = Target.Platform;
++			UnrealTargetConfiguration Configuration = Target.Configuration;
++
++			string UProjectPath = "";
++			if (IsForeignProject)
++			{
++				UProjectPath = String.Format("\"{0}\"", ProjectTarget.UnrealProjectFilePath!.FullName);
++			}
++
++			PlatformProjectGenerator? ProjGenerator = PlatformProjectGenerators.GetPlatformProjectGenerator(Platform, true);
++
++			string BuildArguments = GetBuildArguments(ProjGenerator, ProjectTarget, Target, UProjectPath);
++
++			return new TargetBuildInfo()
++			{
++				BuildCmd = $"{GetBuildScript("Build")} {BuildArguments}",
++				RebuildCmd = $"{GetBuildScript("Rebuild")} {BuildArguments}",
++				CleanCmd = $"{GetBuildScript("Clean")} {BuildArguments}",
++				PrimaryOutput = Target.Binaries[0].OutputFilePath.FullName,
++				BuildByDefault = bBuildByDefault,
++			};
++		}
++
++		private string GetBuildScript(string Cmd)
++		{
++			DirectoryReference BatchFilesDirectory = DirectoryReference.Combine(Unreal.EngineDirectory, "Build", "BatchFiles");
++			return EscapePath(FileReference.Combine(BatchFilesDirectory, $"{Cmd}.bat").FullName);
++		}
++
++		private string GetBuildArguments(
++			PlatformProjectGenerator? ProjGenerator, ProjectTarget ProjectTarget, UEBuildTarget Target, string UProjectPath)
++		{
++			UnrealTargetConfiguration Configuration = Target.Configuration;
++			UnrealTargetPlatform Platform = Target.Platform;
++			string TargetName = ProjectTarget.TargetFilePath.GetFileNameWithoutAnyExtensions();
++			TargetRules TargetRulesObject = ProjectTarget.TargetRules!;
++
++			// This is the standard UE based project NMake build line:
++			//	..\..\Build\BatchFiles\Build.bat <TARGETNAME> <PLATFORM> <CONFIGURATION>
++			//	ie ..\..\Build\BatchFiles\Build.bat BlankProgram Win64 Debug
++
++			StringBuilder BuildArguments = new StringBuilder();
++
++			BuildArguments.AppendFormat("{0} {1} {2}", TargetName, Configuration, Platform);
++			if (IsForeignProject)
++			{
++				BuildArguments.AppendFormat(" -Project={0}", UProjectPath);
++			}
++
++			List<string> ExtraTargets = new List<string>();
++			if (!bUsePrecompiled)
++			{
++				if (TargetRulesObject.Type == TargetType.Editor && Settings.bEditorDependsOnShaderCompileWorker && !Unreal.IsEngineInstalled())
++				{
++					ExtraTargets.Add("ShaderCompileWorker Win64 Development");
++				}
++				if (TargetRulesObject.bWithLiveCoding && Settings.bBuildLiveCodingConsole && !Unreal.IsEngineInstalled() && TargetRulesObject.Name != "LiveCodingConsole")
++				{
++					ExtraTargets.Add(TargetRulesObject.bUseDebugLiveCodingConsole ? "LiveCodingConsole Win64 Debug" : "LiveCodingConsole Win64 Development");
++				}
++			}
++
++			if (ExtraTargets.Count > 0)
++			{
++				BuildArguments.Replace("\"", "\\\"");
++				BuildArguments.Insert(0, "-Target=\"");
++				BuildArguments.Append("\"");
++				foreach (string ExtraTarget in ExtraTargets)
++				{
++					BuildArguments.AppendFormat(" -Target=\"{0} -Quiet\"", ExtraTarget);
++				}
++			}
++
++			if (bUsePrecompiled)
++			{
++				BuildArguments.Append(" -UsePrecompiled");
++			}
++
++			// Always wait for the mutex between UBT invocations, so that building the whole solution doesn't fail.
++			BuildArguments.Append(" -WaitMutex");
++
++			// Always include a flag to format log messages for MSBuild
++			BuildArguments.Append(" -FromMsBuild");
++
++			if (Settings.bAddFastPDBToProjects)
++			{
++				// Pass Fast PDB option to make use of Visual Studio's /DEBUG:FASTLINK option
++				BuildArguments.Append(" -FastPDB");
++			}
++
++			if (ProjGenerator != null)
++			{
++				BuildArguments.Append(ProjGenerator.GetExtraBuildArguments(Platform, Configuration));
++			}
++
++			return BuildArguments.ToString();
++		}
++
++		private static void WriteCommand(JsonWriter Writer, string CommandName, string Command, string Arguments)
++		{
++			Writer.WriteObjectStart(CommandName);
++			Writer.WriteValue("Command", Command);
++			Writer.WriteValue("Args", Arguments);
++			Writer.WriteObjectEnd();
++		}
++
++		internal class ExportedTargetInfo
++		{
++			public string TargetName { get; set; } = String.Empty;
++			public string TargetPath { get; set; } = String.Empty;
++			public string ProjectPath { get; set; } = String.Empty;
++			public string TargetType { get; set; } = String.Empty;
++			public string Platform { get; set; } = String.Empty;
++			public string Configuration { get; set; } = String.Empty;
++			public TargetBuildInfo? BuildInfo { get; set; }
++			public Dictionary<string, ExportedModuleInfo> ModuleToCompileSettings { get; set; } = new Dictionary<string, ExportedModuleInfo>();
++			public Dictionary<string, string> DirToModule { get; set; } = new Dictionary<string, string>();
++		}
++
++		internal class ExportedModuleInfo
++		{
++			public string Name { get; set; } = String.Empty;
++			public string Directory { get; set; } = String.Empty;
++			public string Rules { get; set; } = String.Empty;
++			public string GeneratedCodeDirectory { get; set; } = String.Empty;
++			public List<string> IncludePaths { get; set; } = new List<string>();
++			public List<string> Defines { get; set; } = new List<string>();
++			public string? Standard { get; set; }
++			public List<string> ForcedIncludes { get; set; } = new List<string>();
++			public string? CompilerPath { get; set; }
++			public List<string> CompilerArgs { get; set; } = new List<string>();
++			public Dictionary<string, List<string>> CompilerAdditionalArgs { get; set; } = new Dictionary<string, List<string>>();
++			public string? WindowsSdkVersion { get; set; }
++		}
++
++		internal class TargetBuildInfo
++		{
++			public string BuildCmd { get; set; } = String.Empty;
++			public string RebuildCmd { get; set; } = String.Empty;
++			public string CleanCmd { get; set; } = String.Empty;
++			public string PrimaryOutput { get; set; } = String.Empty;
++			public bool BuildByDefault { get; internal set; }
++		}
++	}
++}
+diff --git a/Engine/Source/Programs/UnrealBuildTool/ProjectFiles/VisualStudioWorkspace/VSWorkspaceProjectFileGenerator.cs b/Engine/Source/Programs/UnrealBuildTool/ProjectFiles/VisualStudioWorkspace/VSWorkspaceProjectFileGenerator.cs
+new file mode 100644
+index 0000000000000..454e291171e00
+--- /dev/null
++++ b/Engine/Source/Programs/UnrealBuildTool/ProjectFiles/VisualStudioWorkspace/VSWorkspaceProjectFileGenerator.cs
+@@ -0,0 +1,248 @@
++// Copyright Epic Games, Inc. All Rights Reserved.
++
++using System;
++using System.Collections.Generic;
++using System.IO;
++using System.Linq;
++using System.Text.Json;
++using EpicGames.Core;
++using Microsoft.Extensions.Logging;
++using UnrealBuildBase;
++
++namespace UnrealBuildTool
++{
++	class VSWorkspaceProjectFileGenerator : ProjectFileGenerator
++	{
++		public override string ProjectFileExtension => ".json";
++
++		// These properties are used by Visual Studio to determine where to read the project files.
++		// So they must remain constant.
++		private const string VSUnrealWorkspaceFileName = ".vs-unreal-workspace";
++		private const string ProjectFilesFolder = "VisualStudio";
++
++		private readonly CommandLineArguments Arguments;
++
++		/// <summary>
++		/// List of deprecated platforms.
++		/// Don't generate project model for these platforms unless they are specified in "Platforms" console arguments.
++		/// </summary>
++		/// <returns></returns>
++		private readonly HashSet<UnrealTargetPlatform> DeprecatedPlatforms = new HashSet<UnrealTargetPlatform>();
++
++		/// <summary>
++		/// Platforms to generate project files for
++		/// </summary>
++		[CommandLine("-Platforms=", ListSeparator = '+')]
++		HashSet<UnrealTargetPlatform> Platforms = new HashSet<UnrealTargetPlatform>();
++
++		/// <summary>
++		/// Target types to generate project files for
++		/// </summary>
++		[CommandLine("-TargetTypes=", ListSeparator = '+')]
++		HashSet<TargetType> TargetTypes = new HashSet<TargetType>();
++
++		/// <summary>
++		/// Target configurations to generate project files for
++		/// </summary>
++		[CommandLine("-TargetConfigurations=", ListSeparator = '+')]
++		HashSet<UnrealTargetConfiguration> TargetConfigurations = new HashSet<UnrealTargetConfiguration>();
++
++		/// <summary>
++		/// Projects to generate project files for
++		/// </summary>
++		[CommandLine("-ProjectNames=", ListSeparator = '+')]
++		HashSet<string> ProjectNames = new HashSet<string>();
++
++		/// <summary>
++		/// Should format JSON files in human readable form, or use packed one without indents
++		/// </summary>
++		[CommandLine("-Minimize", Value = "Compact")]
++		private JsonWriterStyle Minimize = JsonWriterStyle.Readable;
++
++		/// <summary>
++		/// The settings object
++		/// </summary>
++		protected VCProjectFileSettings Settings = new VCProjectFileSettings();
++
++		public VSWorkspaceProjectFileGenerator(FileReference? InOnlyGameProject,
++			CommandLineArguments InArguments)
++			: base(InOnlyGameProject)
++		{
++			Arguments = InArguments;
++			Arguments.ApplyTo(this);
++			XmlConfig.ApplyTo(Settings);
++		}
++
++		public override bool ShouldGenerateIntelliSenseData() => true;
++
++		public override void CleanProjectFiles(DirectoryReference InPrimaryProjectDirectory, string InPrimaryProjectName,
++			DirectoryReference InIntermediateProjectFilesDirectory)
++		{
++			DirectoryReference.Delete(InPrimaryProjectDirectory);
++		}
++
++		protected override void ConfigureProjectFileGeneration(string[] Arguments, ref bool IncludeAllPlatforms)
++		{
++			base.ConfigureProjectFileGeneration(Arguments, ref IncludeAllPlatforms);
++		}
++
++		protected override ProjectFile AllocateProjectFile(FileReference InitFilePath, DirectoryReference BaseDir)
++		{
++			VSWorkspaceProjectFile projectFile = new VSWorkspaceProjectFile(InitFilePath, BaseDir, InitFilePath.Directory,
++				TargetTypes, Arguments, Settings, bUsePrecompiled);
++			return projectFile;
++		}
++
++		protected override bool WriteProjectFiles(PlatformProjectGeneratorCollection PlatformProjectGenerators)
++		{
++			using ProgressWriter Progress = new ProgressWriter("Writing project files...", true);
++			List<ProjectFile> ProjectsToGenerate = new List<ProjectFile>(GeneratedProjectFiles);
++			if (ProjectNames.Any())
++			{
++				ProjectsToGenerate = ProjectsToGenerate.Where(Project =>
++					ProjectNames.Contains(Project.ProjectFilePath.GetFileNameWithoutAnyExtensions())).ToList();
++			}
++
++			int TotalProjectFileCount = ProjectsToGenerate.Count;
++
++			HashSet<UnrealTargetPlatform> PlatformsToGenerate = new HashSet<UnrealTargetPlatform>(SupportedPlatforms);
++			if (Platforms.Any())
++			{
++				PlatformsToGenerate.IntersectWith(Platforms);
++			}
++
++			List<UnrealTargetPlatform> FilteredPlatforms = PlatformsToGenerate.Where(Platform =>
++			{
++				// Skip deprecated unless explicitly specified in the command line.
++				return (!DeprecatedPlatforms.Contains(Platform) || Platforms.Contains(Platform))
++					   && UEBuildPlatform.IsPlatformAvailable(Platform);
++			}).ToList();
++
++			HashSet<UnrealTargetConfiguration> ConfigurationsToGenerate = new HashSet<UnrealTargetConfiguration>(SupportedConfigurations);
++			if (TargetConfigurations.Any())
++			{
++				ConfigurationsToGenerate.IntersectWith(TargetConfigurations);
++			}
++
++			List<UnrealTargetConfiguration> Configurations = ConfigurationsToGenerate.ToList();
++
++			for (int ProjectFileIndex = 0; ProjectFileIndex < ProjectsToGenerate.Count; ++ProjectFileIndex)
++			{
++				VSWorkspaceProjectFile? CurrentProject = ProjectsToGenerate[ProjectFileIndex] as VSWorkspaceProjectFile;
++				if (CurrentProject == null)
++				{
++					return false;
++				}
++
++				if (!CurrentProject.WriteProjectFile(FilteredPlatforms, Configurations, PlatformProjectGenerators, Minimize))
++				{
++					return false;
++				}
++
++				Progress.Write(ProjectFileIndex + 1, TotalProjectFileCount);
++			}
++
++			Progress.Write(TotalProjectFileCount, TotalProjectFileCount);
++
++			return true;
++		}
++
++		public override bool GenerateProjectFiles(PlatformProjectGeneratorCollection PlatformProjectGenerators,
++			String[] Arguments)
++		{
++			bool IncludeAllPlatforms = true;
++			ConfigureProjectFileGeneration(Arguments, ref IncludeAllPlatforms);
++
++			if (bGeneratingGameProjectFiles || Unreal.IsEngineInstalled())
++			{
++				PrimaryProjectPath = OnlyGameProject!.Directory;
++				PrimaryProjectName = OnlyGameProject.GetFileNameWithoutExtension();
++
++				IntermediateProjectFilesPath =
++					DirectoryReference.Combine(PrimaryProjectPath, "Intermediate", "ProjectFiles");
++			}
++
++			SetupSupportedPlatformsAndConfigurations(IncludeAllPlatforms: true, out string SupportedPlatformNames);
++			Log.TraceVerbose("Supported platforms: {Platforms}", SupportedPlatformNames);
++
++			List<FileReference> AllGames = FindGameProjects();
++
++			{
++				// Find all of the target files.
++				List<FileReference> AllTargetFiles = DiscoverTargets(AllGames);
++
++				// If there are multiple targets of a given type for a project, use the order to determine which one gets a suffix.
++				AllTargetFiles = AllTargetFiles.OrderBy(x => x.FullName, StringComparer.OrdinalIgnoreCase).ToList();
++
++				AddProjectsForAllTargets(
++					PlatformProjectGenerators,
++					AllGames,
++					AllTargetFiles,
++					Arguments,
++					out ProjectFile? EngineProject,
++					out List<ProjectFile>? GameProjects,
++					out Dictionary<FileReference, ProjectFile>? ProgramProjects,
++					out Dictionary<RulesAssembly, DirectoryReference>? RulesAssemblies);
++
++				AddAllGameProjects(GameProjects);
++			}
++
++			WriteProjectFiles(PlatformProjectGenerators);
++			WritePrimaryProjectFile(UBTProject, PlatformProjectGenerators);
++
++			return true;
++		}
++
++		protected override bool WritePrimaryProjectFile(ProjectFile? UBTProject,
++			PlatformProjectGeneratorCollection PlatformProjectGenerators)
++		{
++			try
++			{
++				FileReference PrimaryProjectFile = FileReference.Combine(
++					IntermediateProjectFilesPath, ProjectFilesFolder, VSUnrealWorkspaceFileName);
++
++				DirectoryReference.CreateDirectory(PrimaryProjectFile.Directory);
++
++				// Collect all the resulting project files and aggregate the target-level data
++				var AggregatedProjectInfo = GeneratedProjectFiles
++					.Where(Project => Project is VSWorkspaceProjectFile)
++					.OfType<VSWorkspaceProjectFile>()
++					.SelectMany(Project => Project.ExportedTargetProjects)
++					.GroupBy(TargetProject => TargetProject.TargetName)
++					.Select(g => (g.Key, Target: new
++					{
++						TargetType = g.Select(i => i.TargetType).Distinct().Single(),
++						TargetPath = g.Select(i => i.TargetPath).Distinct().Single(),
++						ProjectPath = g.Select(i => i.ProjectPath).Distinct().Single(),
++						Configurations = g.Select(i => i.Configuration).Distinct().ToList(),
++						Platforms = g.Select(i => i.Platform).Distinct().ToList(),
++					}));
++
++				// The inner Targets object is needed for schema compatibility with the Query Mode API.
++				var Result = new
++				{
++					Targets = AggregatedProjectInfo.ToDictionary(item => item.Key, item => item.Target)
++				};
++
++				File.WriteAllText(PrimaryProjectFile.FullName, JsonSerializer.Serialize(Result, new JsonSerializerOptions
++				{
++					PropertyNamingPolicy = JsonNamingPolicy.CamelCase,
++					WriteIndented = Minimize == JsonWriterStyle.Readable,
++				}));
++			}
++			catch (Exception Ex)
++			{
++				Log.TraceWarning("Exception while writing root project file: {Ex}", Ex.ToString());
++				return false;
++			}
++
++			return true;
++		}
++
++		/// <inheritdoc />
++		protected override FileReference GetProjectLocation(string BaseName)
++		{
++			return FileReference.Combine(IntermediateProjectFilesPath, ProjectFilesFolder, BaseName + ProjectFileExtension);
++		}
++	}
++}
+diff --git a/Engine/Source/Programs/UnrealBuildTool/ToolChain/UEToolChain.cs b/Engine/Source/Programs/UnrealBuildTool/ToolChain/UEToolChain.cs
+index e8b40fff208d9..4f4fff19c8551 100644
+--- a/Engine/Source/Programs/UnrealBuildTool/ToolChain/UEToolChain.cs
++++ b/Engine/Source/Programs/UnrealBuildTool/ToolChain/UEToolChain.cs
+@@ -34,6 +34,12 @@ public static DirectoryReference GetModuleInterfaceDir(DirectoryReference Output
+ 			return DirectoryReference.Combine(OutputDir, "Ifc");
+ 		}
+ 
++		// Return the path to the cpp compiler that will be used by this toolchain.
++		public virtual FileReference? GetCppCompilerPath()
++		{
++			return null;
++		}
++
+ 		public abstract CPPOutput CompileCPPFiles(CppCompileEnvironment CompileEnvironment, List<FileItem> InputFiles, DirectoryReference OutputDir, string ModuleName, IActionGraphBuilder Graph);
+ 
+ 		public virtual CPPOutput CompileRCFiles(CppCompileEnvironment Environment, List<FileItem> InputFiles, DirectoryReference OutputDir, IActionGraphBuilder Graph)
+@@ -65,6 +71,20 @@ public virtual FileItem[] LinkAllFiles(LinkEnvironment LinkEnvironment, bool bBu
+ 			return LinkFile != null ? new FileItem[] { LinkFile } : new FileItem[] { };
+ 		}
+ 
++		public virtual IEnumerable<string> GetGlobalCommandLineArgs(CppCompileEnvironment CompileEnvironment)
++		{
++			return Array.Empty<string>();
++		}
++
++		public virtual IEnumerable<string> GetCPPCommandLineArgs(CppCompileEnvironment CompileEnvironment)
++		{
++			return Array.Empty<string>();
++		}
++
++		public virtual IEnumerable<string> GetCCommandLineArgs(CppCompileEnvironment CompileEnvironment)
++		{
++			return Array.Empty<string>();
++		}
+ 
+ 		/// <summary>
+ 		/// Get the name of the response file for the current linker environment and output file
diff --git a/Engine/Extras/VisualStudioWorkspace/UnrealBuildTool-5.1.patch b/Engine/Extras/VisualStudioWorkspace/UnrealBuildTool-5.1.patch
new file mode 100644
--- /dev/null
+++ b/Engine/Extras/VisualStudioWorkspace/UnrealBuildTool-5.1.patch
@@ -0,0 +1,739 @@
+From a0a3568806c71d9839a1f0163fd2b0c235721641 Mon Sep 17 00:00:00 2001
+From: [REDACTED_PERSON] <[REDACTED_EMAIL]>
+Date: Wed, 17 Jan 2024 19:49:54 -0800
+Subject: [PATCH] [5.1] Project Generator for VS Workspace support
+
+Adding a new Project Generator for Visual Studio that will be used for
+the new support to opening `.uproject` files directly in the IDE without
+generating the traditional `.sln` and .vcxproj files.
+
+This is based on the QueryMode API in `ue5-main` and adapted to the
+GenerateProjectFiles flow. The project files will contain information for each
+Target+Configuration+Platform combinations that are used the project.
+
+Visual Studio will automatically invoke this new generator during the flow to
+open a `.uproject` file directly.
+---
+ .../Modes/GenerateProjectFilesMode.cs         |   4 +
+ .../Platform/Windows/VCToolChain.cs           |  26 ++
+ .../ProjectFiles/ProjectFileGenerator.cs      |   5 +-
+ .../VSWorkspaceProjectFile.cs                 | 313 ++++++++++++++++++
+ .../VSWorkspaceProjectFileGenerator.cs        | 253 ++++++++++++++
+ .../UnrealBuildTool/ToolChain/UEToolChain.cs  |  21 ++
+ 6 files changed, 620 insertions(+), 2 deletions(-)
+ create mode 100644 Engine/Source/Programs/UnrealBuildTool/ProjectFiles/VisualStudioWorkspace/VSWorkspaceProjectFile.cs
+ create mode 100644 Engine/Source/Programs/UnrealBuildTool/ProjectFiles/VisualStudioWorkspace/VSWorkspaceProjectFileGenerator.cs
+
+diff --git a/Engine/Source/Programs/UnrealBuildTool/Modes/GenerateProjectFilesMode.cs b/Engine/Source/Programs/UnrealBuildTool/Modes/GenerateProjectFilesMode.cs
+index adf27b9916eb5..186d881fb18c1 100644
+--- a/Engine/Source/Programs/UnrealBuildTool/Modes/GenerateProjectFilesMode.cs
++++ b/Engine/Source/Programs/UnrealBuildTool/Modes/GenerateProjectFilesMode.cs
+@@ -36,6 +36,7 @@ class GenerateProjectFilesMode : ToolMode
+ 		[CommandLine("-EddieProjectFiles", Value = nameof(ProjectFileFormat.Eddie))]
+ 		[CommandLine("-VSCode", Value = nameof(ProjectFileFormat.VisualStudioCode))]
+ 		[CommandLine("-VSMac", Value = nameof(ProjectFileFormat.VisualStudioMac))]
++		[CommandLine("-VSWorkspace", Value = nameof(ProjectFileFormat.VisualStudioWorkspace))]
+ 		[CommandLine("-CLion", Value = nameof(ProjectFileFormat.CLion))]
+ 		[CommandLine("-Rider", Value = nameof(ProjectFileFormat.Rider))]
+ 		#if __VPROJECT_AVAILABLE__
+@@ -204,6 +205,9 @@ public override int Execute(CommandLineArguments Arguments, ILogger Logger)
+ 					case ProjectFileFormat.VisualStudioCode:
+ 						Generator = new VSCodeProjectFileGenerator(ProjectFile);
+ 						break;
++					case ProjectFileFormat.VisualStudioWorkspace:
++						Generator = new VSWorkspaceProjectFileGenerator(ProjectFile, Arguments);
++						break;
+ 					case ProjectFileFormat.CLion:
+ 						Generator = new CLionGenerator(ProjectFile);
+ 						break;
+diff --git a/Engine/Source/Programs/UnrealBuildTool/Platform/Windows/VCToolChain.cs b/Engine/Source/Programs/UnrealBuildTool/Platform/Windows/VCToolChain.cs
+index 2f44e14fb32ab..e10d8646f0b23 100644
+--- a/Engine/Source/Programs/UnrealBuildTool/Platform/Windows/VCToolChain.cs
++++ b/Engine/Source/Programs/UnrealBuildTool/Platform/Windows/VCToolChain.cs
+@@ -55,6 +55,11 @@ public VCToolChain(ReadOnlyTargetRules Target, ILogger Logger)
+ 			}
+ 		}
+ 
++		public override FileReference? GetCppCompilerPath()
++		{
++			return EnvVars.CompilerPath;
++		}
++
+ 		/// <summary>
+ 		/// Prepares the environment for building
+ 		/// </summary>
+@@ -250,6 +255,27 @@ public static void AddSourceDependsFile(List<string> Arguments, FileItem SourceD
+ 			Arguments.Add($"/clang:-MD /clang:-MF\"{SourceDependsFileString}\"");
+ 		}
+ 
++		public override IEnumerable<string> GetGlobalCommandLineArgs(CppCompileEnvironment CompileEnvironment)
++		{
++			List<string> Arguments = new();
++			AppendCLArguments_Global(new(CompileEnvironment), Arguments);
++			return Arguments;
++		}
++
++		public override IEnumerable<string> GetCPPCommandLineArgs(CppCompileEnvironment CompileEnvironment)
++		{
++			List<string> Arguments = new();
++			AppendCLArguments_CPP(CompileEnvironment, Arguments);
++			return Arguments;
++		}
++
++		public override IEnumerable<string> GetCCommandLineArgs(CppCompileEnvironment CompileEnvironment)
++		{
++			List<string> Arguments = new();
++			AppendCLArguments_C(CompileEnvironment, Arguments);
++			return Arguments;
++		}
++
+ 		protected virtual void AppendCLArguments_Global(CppCompileEnvironment CompileEnvironment, List<string> Arguments)
+ 		{
+ 			// Suppress generation of object code for unreferenced inline functions. Enabling this option is more standards compliant, and causes a big reduction
+diff --git a/Engine/Source/Programs/UnrealBuildTool/ProjectFiles/ProjectFileGenerator.cs b/Engine/Source/Programs/UnrealBuildTool/ProjectFiles/ProjectFileGenerator.cs
+index bb6c18322a31b..2dd4cea450647 100644
+--- a/Engine/Source/Programs/UnrealBuildTool/ProjectFiles/ProjectFileGenerator.cs
++++ b/Engine/Source/Programs/UnrealBuildTool/ProjectFiles/ProjectFileGenerator.cs
+@@ -133,6 +133,7 @@ enum ProjectFileFormat
+ 		VisualStudio,
+ 		VisualStudio2019,
+ 		VisualStudio2022,
++		VisualStudioWorkspace,
+ 		XCode,
+ 		Eddie,
+ 		VisualStudioCode,
+@@ -2332,7 +2333,7 @@ private ProjectFile FindProjectForModule(FileReference CurModuleFile, List<FileR
+ 		/// <param name="ProgramProjects">Map of program names to all of the program projects we created</param>
+ 		/// <param name="RulesAssemblies">Map of RuleAssemblies to their base folders</param>
+ 		/// <param name="Logger">Logger for output</param>
+-		private void AddProjectsForAllTargets(
++		protected void AddProjectsForAllTargets(
+ 			PlatformProjectGeneratorCollection PlatformProjectGenerators,
+ 			List<FileReference> AllGames,
+ 			List<FileReference> AllTargetFiles,
+@@ -2634,7 +2635,7 @@ protected void AddProjectsForMods(List<ProjectFile> GameProjects, List<ProjectFi
+ 		/// </summary>
+ 		/// <param name="BaseName">The base name for the project file</param>
+ 		/// <returns>Full path to the project file</returns>
+-		protected FileReference GetProjectLocation(string BaseName)
++		protected virtual FileReference GetProjectLocation(string BaseName)
+ 		{
+ 			return FileReference.Combine(IntermediateProjectFilesPath, BaseName + ProjectFileExtension);
+ 		}
+diff --git a/Engine/Source/Programs/UnrealBuildTool/ProjectFiles/VisualStudioWorkspace/VSWorkspaceProjectFile.cs b/Engine/Source/Programs/UnrealBuildTool/ProjectFiles/VisualStudioWorkspace/VSWorkspaceProjectFile.cs
+new file mode 100644
+index 0000000000000..c528104137af2
+--- /dev/null
++++ b/Engine/Source/Programs/UnrealBuildTool/ProjectFiles/VisualStudioWorkspace/VSWorkspaceProjectFile.cs
+@@ -0,0 +1,313 @@
++// Copyright Epic Games, Inc. All Rights Reserved.
++
++using System;
++using System.Collections.Generic;
++using System.IO;
++using System.Linq;
++using System.Text;
++using System.Text.Json;
++using EpicGames.Core;
++using Microsoft.Extensions.Logging;
++using UnrealBuildBase;
++
++namespace UnrealBuildTool
++{
++	internal class VSWorkspaceProjectFile : ProjectFile
++	{
++		private readonly DirectoryReference RootPath;
++		private readonly HashSet<TargetType> TargetTypes;
++		private readonly CommandLineArguments Arguments;
++
++		/// <summary>
++		/// Collection of output files for this project
++		/// </summary>
++		public List<ExportedTargetInfo> ExportedTargetProjects { get; set; } = new();
++
++		public VSWorkspaceProjectFile(FileReference InProjectFilePath, DirectoryReference BaseDir,
++			DirectoryReference RootPath, HashSet<TargetType> TargetTypes, CommandLineArguments Arguments)
++			: base(InProjectFilePath, BaseDir)
++		{
++			this.RootPath = RootPath;
++			this.TargetTypes = TargetTypes;
++			this.Arguments = Arguments;
++		}
++
++		/// <summary>
++		/// Write project file info in JSON file.
++		/// For every combination of <c>UnrealTargetPlatform</c>, <c>UnrealTargetConfiguration</c> and <c>TargetType</c>
++		/// will be generated separate JSON file.
++		/// </summary>
++		public bool WriteProjectFile(List<UnrealTargetPlatform> InPlatforms,
++			List<UnrealTargetConfiguration> InConfigurations,
++			PlatformProjectGeneratorCollection PlatformProjectGenerators, JsonWriterStyle Minimize, ILogger Logger)
++		{
++			DirectoryReference ProjectRootFolder = RootPath;
++
++			Dictionary<FileReference, (UEBuildTarget BuildTarget, bool bBuildByDefault)> FileToTarget = new();
++
++			foreach (UnrealTargetPlatform Platform in InPlatforms)
++			{
++				foreach (UnrealTargetConfiguration Configuration in InConfigurations)
++				{
++					foreach (ProjectTarget ProjectTarget in ProjectTargets.OfType<ProjectTarget>())
++					{
++						if (TargetTypes.Any() && !TargetTypes.Contains(ProjectTarget.TargetRules!.Type))
++						{
++							continue;
++						}
++
++						// Skip Programs for all configs except for current platform + Development & Debug configurations
++						if (ProjectTarget.TargetRules!.Type == TargetType.Program &&
++							(BuildHostPlatform.Current.Platform != Platform ||
++							 !(Configuration == UnrealTargetConfiguration.Development || Configuration == UnrealTargetConfiguration.Debug)))
++						{
++							continue;
++						}
++
++						// Skip Editor for all platforms except for current platform
++						if (ProjectTarget.TargetRules.Type == TargetType.Editor &&
++							(BuildHostPlatform.Current.Platform != Platform ||
++							(Configuration == UnrealTargetConfiguration.Test || Configuration == UnrealTargetConfiguration.Shipping)))
++						{
++							continue;
++						}
++
++						bool bBuildByDefault = ShouldBuildByDefaultForSolutionTargets && ProjectTarget.SupportedPlatforms.Contains(Platform);
++
++						string DefaultArchitecture = UEBuildPlatform
++							.GetBuildPlatform(Platform)
++							.GetDefaultArchitecture(ProjectTarget.UnrealProjectFilePath);
++
++						TargetDescriptor TargetDesc = new(ProjectTarget.UnrealProjectFilePath, ProjectTarget.Name,
++							Platform, Configuration, DefaultArchitecture, Arguments);
++
++						try
++						{
++							FileReference OutputFile = FileReference.Combine(ProjectRootFolder,
++								$"{ProjectTarget.TargetFilePath.GetFileNameWithoutAnyExtensions()}_{Configuration}_{Platform}.json");
++
++							UEBuildTarget BuildTarget = UEBuildTarget.Create(TargetDesc, false, false, false, Logger);
++							FileToTarget.Add(OutputFile, (BuildTarget, bBuildByDefault));
++						}
++						catch (Exception Ex)
++						{
++							Logger.LogWarning("Exception while generating include data for Target:{Target}, Platform: {Platform}, Configuration: {Configuration}", TargetDesc.Name, Platform.ToString(), Configuration.ToString());
++							Logger.LogWarning("{Ex}", Ex.ToString());
++						}
++					}
++				}
++			}
++
++			foreach (var Entry in FileToTarget)
++			{
++				var OutputFile = Entry.Key;
++				var BuildTarget = Entry.Value.BuildTarget;
++				var bBuildByDefault = Entry.Value.bBuildByDefault;
++
++				try
++				{
++					BuildTarget.PreBuildSetup(Logger);
++
++					ExportedTargetInfo TargetInfo = ExportTarget(BuildTarget, bBuildByDefault, PlatformProjectGenerators, Logger);
++
++					DirectoryReference.CreateDirectory(OutputFile.Directory);
++					using FileStream Stream = new(OutputFile.FullName, FileMode.Create, FileAccess.Write);
++					JsonSerializer.Serialize(Stream, TargetInfo, options: new JsonSerializerOptions()
++					{
++						PropertyNamingPolicy = JsonNamingPolicy.CamelCase,
++						WriteIndented = Minimize == JsonWriterStyle.Readable,
++					});
++
++					ExportedTargetProjects.Add(TargetInfo);
++				}
++				catch (Exception Ex)
++				{
++					Logger.LogWarning("Exception while generating include data for Target:{Target}, Platform: {Platform}, Configuration: {Configuration}",
++						BuildTarget.AppName, BuildTarget.Platform.ToString(), BuildTarget.Configuration.ToString());
++					Logger.LogWarning("{Ex}", Ex.ToString());
++				}
++			}
++
++			return true;
++		}
++
++		private ExportedTargetInfo ExportTarget(
++			UEBuildTarget Target, bool bBuildByDefault, PlatformProjectGeneratorCollection PlatformProjectGenerators, ILogger Logger)
++		{
++			ExportedTargetInfo TargetInfo = new()
++			{
++				TargetName = Target.TargetName,
++				TargetPath = Target.TargetRulesFile.FullName,
++				ProjectPath = Target.ProjectFile?.FullName ?? String.Empty,
++				TargetType = Target.TargetType.ToString(),
++				Platform = Target.Platform.ToString(),
++				Configuration = Target.Configuration.ToString(),
++				BuildInfo = ExportBuildInfo(Target, PlatformProjectGenerators, bBuildByDefault, Logger)
++			};
++
++			UEToolChain TargetToolChain = Target.CreateToolchain(Target.Platform);
++			CppCompileEnvironment GlobalCompileEnvironment = Target.CreateCompileEnvironmentForProjectFiles(Logger);
++
++			HashSet<string> ModuleNames = new();
++			foreach (UEBuildBinary Binary in Target.Binaries)
++			{
++				CppCompileEnvironment BinaryCompileEnvironment = Binary.CreateBinaryCompileEnvironment(GlobalCompileEnvironment);
++				IEnumerable<UEBuildModuleCPP> CandidateModules = Binary.Modules.Where(x => x is UEBuildModuleCPP).Cast<UEBuildModuleCPP>();
++
++				foreach (var ModuleCpp in CandidateModules)
++				{
++					if (!ModuleNames.Add(ModuleCpp.Name))
++					{
++						continue;
++					}
++
++					CppCompileEnvironment ModuleCompileEnvironment = ModuleCpp.CreateCompileEnvironmentForIntellisense(Target.Rules, BinaryCompileEnvironment, Logger);
++					TargetInfo.ModuleToCompileSettings.Add(ModuleCpp.Name, ExportModule(ModuleCpp, TargetToolChain, ModuleCompileEnvironment, Logger));
++
++					foreach (DirectoryReference ModuleDirectory in ModuleCpp.ModuleDirectories)
++					{
++						TargetInfo.DirToModule.TryAdd(ModuleDirectory.FullName, ModuleCpp.Name);
++					}
++
++					if (ModuleCpp.GeneratedCodeDirectory != null)
++					{
++						TargetInfo.DirToModule.TryAdd(ModuleCpp.GeneratedCodeDirectory.FullName, ModuleCpp.Name);
++					}
++				}
++			}
++
++			return TargetInfo;
++		}
++
++		private static ExportedModuleInfo ExportModule(UEBuildModuleCPP Module, UEToolChain TargetToolChain, CppCompileEnvironment ModuleCompileEnvironment, ILogger Logger)
++		{
++			ExportedModuleInfo Result = new()
++			{
++				Name = Module.Name,
++				Directory = Module.ModuleDirectory.FullName,
++				Rules = Module.RulesFile.FullName,
++				GeneratedCodeDirectory = Module.GeneratedCodeDirectory != null ? Module.GeneratedCodeDirectory.FullName : String.Empty,
++				Standard = ModuleCompileEnvironment.CppStandard.ToString(),
++			};
++
++			Result.IncludePaths.AddRange(Module.PublicIncludePaths.Select(x => x.FullName));
++			Result.IncludePaths.AddRange(Module.PublicSystemIncludePaths.Select(x => x.FullName));
++			Result.IncludePaths.AddRange(Module.InternalIncludePaths.Select(x => x.FullName));
++			Result.IncludePaths.AddRange(Module.LegacyPublicIncludePaths.Select(x => x.FullName));
++			Result.IncludePaths.AddRange(Module.PrivateIncludePaths.Select(x => x.FullName));
++
++			Result.IncludePaths.AddRange(ModuleCompileEnvironment.UserIncludePaths.Select(x => x.FullName));
++
++			Result.IncludePaths.AddRange(Module.PublicSystemLibraryPaths.Select(x => x.FullName));
++			Result.IncludePaths.AddRange(Module.PublicSystemLibraries.Concat(Module.PublicLibraries.Select(x => x.FullName)));
++
++			if (TargetToolChain is VCToolChain TargetVCToolChain
++				&& OperatingSystem.IsWindows())
++			{
++				string VCIncludePaths = VCToolChain.GetVCIncludePaths(ModuleCompileEnvironment.Platform, WindowsCompiler.VisualStudio2022, null, Logger);
++				Result.IncludePaths.AddRange(VCIncludePaths.Split(";"));
++			}
++
++			Result.Defines.AddRange(Module.PublicDefinitions);
++			Result.Defines.AddRange(Module.Rules.PrivateDefinitions);
++			Result.Defines.AddRange(Module.Rules.bTreatAsEngineModule ? Array.Empty<string>() : Module.Rules.Target.ProjectDefinitions);
++			Result.Defines.AddRange(Module.GetEmptyApiMacros());
++
++			var ForcedIncludes = ModuleCompileEnvironment.ForceIncludeFiles.ToList();
++			if (ModuleCompileEnvironment.PrecompiledHeaderAction == PrecompiledHeaderAction.Include)
++			{
++				FileItem IncludeHeader = FileItem.GetItemByFileReference(ModuleCompileEnvironment.PrecompiledHeaderIncludeFilename!);
++				ForcedIncludes.Insert(0, IncludeHeader);
++			}
++
++			Result.ForcedIncludes.AddRange(ForcedIncludes.Select(x => x.FullName));
++
++			Result.CompilerPath = TargetToolChain.GetCppCompilerPath()?.ToString();
++			Result.CompilerArgs.AddRange(TargetToolChain.GetGlobalCommandLineArgs(ModuleCompileEnvironment));
++			Result.CompilerAdditionalArgs.Add("c", TargetToolChain.GetCCommandLineArgs(ModuleCompileEnvironment).ToList());
++			Result.CompilerAdditionalArgs.Add("cpp", TargetToolChain.GetCPPCommandLineArgs(ModuleCompileEnvironment).ToList());
++
++			return Result;
++		}
++
++		private TargetBuildInfo? ExportBuildInfo(UEBuildTarget Target, PlatformProjectGeneratorCollection PlatformProjectGenerators, bool bBuildByDefault, ILogger Logger)
++		{
++			if (!BuildHostPlatform.Current.Platform.IsInGroup(UnrealPlatformGroup.Windows))
++			{
++				Logger.LogWarning("Unsupported platform for Build Information: {Platform}", BuildHostPlatform.Current.Platform.ToString());
++				return null;
++			}
++
++			if (IsStubProject)
++			{
++				return null;
++			}
++
++			ProjectTarget ProjectTarget = ProjectTargets.OfType<ProjectTarget>().Single(It => Target.TargetRulesFile == It.TargetFilePath);
++			UnrealTargetPlatform Platform = Target.Platform;
++			UnrealTargetConfiguration Configuration = Target.Configuration;
++
++			string UProjectPath = "";
++			if (IsForeignProject)
++			{
++				UProjectPath = String.Format("\"{0}\"", ProjectTarget.UnrealProjectFilePath!.FullName);
++			}
++
++			PlatformProjectGenerator? ProjGenerator = PlatformProjectGenerators.GetPlatformProjectGenerator(Platform, true);
++			VCProjectFile.BuildCommandBuilder BuildCommandBuilder = new(Configuration, Platform, ProjectTarget, UProjectPath)
++			{
++				ProjectGenerator = ProjGenerator,
++				bIsForeignProject = IsForeignProject
++			};
++
++			string BuildArguments = BuildCommandBuilder.GetBuildArguments();
++
++			return new TargetBuildInfo()
++			{
++				BuildCmd = $"{BuildCommandBuilder.BuildScript.FullName} {BuildArguments}",
++				RebuildCmd = $"{BuildCommandBuilder.RebuildScript.FullName} {BuildArguments}",
++				CleanCmd = $"{BuildCommandBuilder.CleanScript.FullName} {BuildArguments}",
++				PrimaryOutput = Target.Binaries[0].OutputFilePath.FullName,
++				BuildByDefault = bBuildByDefault,
++			};
++		}
++
++		internal class ExportedTargetInfo
++		{
++			public string TargetName { get; set; } = String.Empty;
++			public string TargetPath { get; set; } = String.Empty;
++			public string ProjectPath { get; set; } = String.Empty;
++			public string TargetType { get; set; } = String.Empty;
++			public string Platform { get; set; } = String.Empty;
++			public string Configuration { get; set; } = String.Empty;
++			public TargetBuildInfo? BuildInfo { get; set; }
++			public Dictionary<string, ExportedModuleInfo> ModuleToCompileSettings { get; set; } = new();
++			public Dictionary<string, string> DirToModule { get; set; } = new();
++		}
++
++		internal class ExportedModuleInfo
++		{
++			public string Name { get; set; } = String.Empty;
++			public string Directory { get; set; } = String.Empty;
++			public string Rules { get; set; } = String.Empty;
++			public string GeneratedCodeDirectory { get; set; } = String.Empty;
++			public List<string> IncludePaths { get; set; } = new();
++			public List<string> Defines { get; set; } = new();
++			public string? Standard { get; set; }
++			public List<string> ForcedIncludes { get; set; } = new();
++			public string? CompilerPath { get; set; }
++			public List<string> CompilerArgs { get; set; } = new();
++			public Dictionary<string, List<string>> CompilerAdditionalArgs { get; set; } = new();
++			public string? WindowsSdkVersion { get; set; }
++		}
++
++		internal class TargetBuildInfo
++		{
++			public string BuildCmd { get; set; } = String.Empty;
++			public string RebuildCmd { get; set; } = String.Empty;
++			public string CleanCmd { get; set; } = String.Empty;
++			public string PrimaryOutput { get; set; } = String.Empty;
++			public bool BuildByDefault { get; internal set; }
++		}
++	}
++}
+diff --git a/Engine/Source/Programs/UnrealBuildTool/ProjectFiles/VisualStudioWorkspace/VSWorkspaceProjectFileGenerator.cs b/Engine/Source/Programs/UnrealBuildTool/ProjectFiles/VisualStudioWorkspace/VSWorkspaceProjectFileGenerator.cs
+new file mode 100644
+index 0000000000000..ccc3fee601f11
+--- /dev/null
++++ b/Engine/Source/Programs/UnrealBuildTool/ProjectFiles/VisualStudioWorkspace/VSWorkspaceProjectFileGenerator.cs
+@@ -0,0 +1,253 @@
++// Copyright Epic Games, Inc. All Rights Reserved.
++
++using System;
++using System.Collections.Generic;
++using System.IO;
++using System.Linq;
++using System.Text.Json;
++using EpicGames.Core;
++using Microsoft.Extensions.Logging;
++using UnrealBuildBase;
++
++namespace UnrealBuildTool
++{
++	class VSWorkspaceProjectFileGenerator : ProjectFileGenerator
++	{
++		public override string ProjectFileExtension => ".json";
++
++		// These properties are used by Visual Studio to determine where to read the project files.
++		// So they must remain constant.
++		private const string VSUnrealWorkspaceFileName = ".vs-unreal-workspace";
++		private const string ProjectFilesFolder = "VisualStudio";
++
++		private readonly CommandLineArguments Arguments;
++
++		/// <summary>
++		/// List of deprecated platforms.
++		/// Don't generate project model for these platforms unless they are specified in "Platforms" console arguments.
++		/// </summary>
++		/// <returns></returns>
++		private readonly HashSet<UnrealTargetPlatform> DeprecatedPlatforms = new();
++
++		/// <summary>
++		/// Platforms to generate project files for
++		/// </summary>
++		[CommandLine("-Platforms=", ListSeparator = '+')]
++		HashSet<UnrealTargetPlatform> Platforms = new();
++
++		/// <summary>
++		/// Target types to generate project files for
++		/// </summary>
++		[CommandLine("-TargetTypes=", ListSeparator = '+')]
++		HashSet<TargetType> TargetTypes = new();
++
++		/// <summary>
++		/// Target configurations to generate project files for
++		/// </summary>
++		[CommandLine("-TargetConfigurations=", ListSeparator = '+')]
++		HashSet<UnrealTargetConfiguration> TargetConfigurations = new();
++
++		/// <summary>
++		/// Projects to generate project files for
++		/// </summary>
++		[CommandLine("-ProjectNames=", ListSeparator = '+')]
++		HashSet<string> ProjectNames = new();
++
++		/// <summary>
++		/// Should format JSON files in human readable form, or use packed one without indents
++		/// </summary>
++		[CommandLine("-Minimize", Value = "Compact")]
++		private JsonWriterStyle Minimize = JsonWriterStyle.Readable;
++
++		public VSWorkspaceProjectFileGenerator(FileReference? InOnlyGameProject,
++			CommandLineArguments InArguments)
++			: base(InOnlyGameProject)
++		{
++			Arguments = InArguments;
++			Arguments.ApplyTo(this);
++		}
++
++		public override bool ShouldGenerateIntelliSenseData() => true;
++
++		public override void CleanProjectFiles(DirectoryReference InPrimaryProjectDirectory, string InPrimaryProjectName,
++			DirectoryReference InIntermediateProjectFilesDirectory, ILogger Logger)
++		{
++			DirectoryReference.Delete(InPrimaryProjectDirectory);
++		}
++
++		protected override void ConfigureProjectFileGeneration(string[] Arguments, ref bool IncludeAllPlatforms, ILogger Logger)
++		{
++			base.ConfigureProjectFileGeneration(Arguments, ref IncludeAllPlatforms, Logger);
++		}
++
++		protected override ProjectFile AllocateProjectFile(FileReference InitFilePath, DirectoryReference BaseDir)
++		{
++			VSWorkspaceProjectFile projectFile = new(InitFilePath, BaseDir, RootPath: InitFilePath.Directory,
++				Arguments: Arguments, TargetTypes: TargetTypes);
++			return projectFile;
++		}
++
++		protected override bool WriteProjectFiles(PlatformProjectGeneratorCollection PlatformProjectGenerators, ILogger Logger)
++		{
++			using ProgressWriter Progress = new("Writing project files...", true, Logger);
++			List<ProjectFile> ProjectsToGenerate = new(GeneratedProjectFiles);
++			if (ProjectNames.Any())
++			{
++				ProjectsToGenerate = ProjectsToGenerate.Where(Project =>
++					ProjectNames.Contains(Project.ProjectFilePath.GetFileNameWithoutAnyExtensions())).ToList();
++			}
++
++			int TotalProjectFileCount = ProjectsToGenerate.Count;
++
++			HashSet<UnrealTargetPlatform> PlatformsToGenerate = new(SupportedPlatforms);
++			if (Platforms.Any())
++			{
++				PlatformsToGenerate.IntersectWith(Platforms);
++			}
++
++			List<UnrealTargetPlatform> FilteredPlatforms = PlatformsToGenerate.Where(Platform =>
++			{
++				// Skip deprecated unless explicitly specified in the command line.
++				return (!DeprecatedPlatforms.Contains(Platform) || Platforms.Contains(Platform))
++					   && UEBuildPlatform.IsPlatformAvailable(Platform);
++			}).ToList();
++
++			HashSet<UnrealTargetConfiguration> ConfigurationsToGenerate = new(SupportedConfigurations);
++			if (TargetConfigurations.Any())
++			{
++				ConfigurationsToGenerate.IntersectWith(TargetConfigurations);
++			}
++
++			List<UnrealTargetConfiguration> Configurations = ConfigurationsToGenerate.ToList();
++
++			for (int ProjectFileIndex = 0; ProjectFileIndex < ProjectsToGenerate.Count; ++ProjectFileIndex)
++			{
++				if (ProjectsToGenerate[ProjectFileIndex] is not VSWorkspaceProjectFile CurrentProject)
++				{
++					return false;
++				}
++
++				if (!CurrentProject.WriteProjectFile(FilteredPlatforms, Configurations, PlatformProjectGenerators, Minimize, Logger))
++				{
++					return false;
++				}
++
++				Progress.Write(ProjectFileIndex + 1, TotalProjectFileCount);
++			}
++
++			Progress.Write(TotalProjectFileCount, TotalProjectFileCount);
++
++			return true;
++		}
++
++		public override bool GenerateProjectFiles(PlatformProjectGeneratorCollection PlatformProjectGenerators,
++			String[] Arguments, ILogger Logger)
++		{
++			bool IncludeAllPlatforms = true;
++			ConfigureProjectFileGeneration(Arguments, ref IncludeAllPlatforms, Logger);
++
++			if (bGeneratingGameProjectFiles || Unreal.IsEngineInstalled())
++			{
++				PrimaryProjectPath = OnlyGameProject!.Directory;
++				PrimaryProjectName = OnlyGameProject.GetFileNameWithoutExtension();
++
++				IntermediateProjectFilesPath =
++					DirectoryReference.Combine(PrimaryProjectPath, "Intermediate", "ProjectFiles");
++			}
++
++			SetupSupportedPlatformsAndConfigurations(IncludeAllPlatforms: true, Logger, out string SupportedPlatformNames);
++			Logger.LogDebug("Supported platforms: {Platforms}", SupportedPlatformNames);
++
++			List<FileReference> AllGames = FindGameProjects(Logger);
++
++			{
++				// Find all of the target files.
++				List<FileReference> AllTargetFiles = DiscoverTargets(
++					AllGames,
++					Logger);
++
++				// If there are multiple targets of a given type for a project, use the order to determine which one gets a suffix.
++				AllTargetFiles = AllTargetFiles.OrderBy(x => x.FullName, StringComparer.OrdinalIgnoreCase).ToList();
++
++				ProjectFile? EngineProject = null;
++				List<ProjectFile> GameProjects = new();
++				List<ProjectFile> ModProjects = new();
++				Dictionary<FileReference, ProjectFile> ProgramProjects = new();
++				Dictionary<RulesAssembly, DirectoryReference> RulesAssemblies = new();
++				Dictionary<ProjectFile, FileReference> ProjectFileToUProjectFile = new();
++
++				AddProjectsForAllTargets(
++					PlatformProjectGenerators,
++					AllGames,
++					AllTargetFiles,
++					Arguments,
++					ref EngineProject,
++					GameProjects,
++					ProgramProjects,
++					RulesAssemblies,
++					Logger);
++
++				AddAllGameProjects(GameProjects);
++			}
++
++			WriteProjectFiles(PlatformProjectGenerators, Logger);
++			WritePrimaryProjectFile(UBTProject, PlatformProjectGenerators, Logger);
++
++			return true;
++		}
++
++		protected override bool WritePrimaryProjectFile(ProjectFile? UBTProject,
++			PlatformProjectGeneratorCollection PlatformProjectGenerators,
++			ILogger Logger)
++		{
++			try
++			{
++				FileReference PrimaryProjectFile = FileReference.Combine(
++					IntermediateProjectFilesPath, ProjectFilesFolder, VSUnrealWorkspaceFileName);
++
++				DirectoryReference.CreateDirectory(PrimaryProjectFile.Directory);
++
++				// Collect all the resulting project files and aggregate the target-level data
++				var AggregatedProjectInfo = GeneratedProjectFiles
++					.Where(Project => Project is VSWorkspaceProjectFile)
++					.OfType<VSWorkspaceProjectFile>()
++					.SelectMany(Project => Project.ExportedTargetProjects)
++					.GroupBy(TargetProject => TargetProject.TargetName)
++					.Select(g => (g.Key, Target: new
++					{
++						TargetType = g.Select(i => i.TargetType).Distinct().Single(),
++						TargetPath = g.Select(i => i.TargetPath).Distinct().Single(),
++						ProjectPath = g.Select(i => i.ProjectPath).Distinct().Single(),
++						Configurations = g.Select(i => i.Configuration).Distinct().ToList(),
++						Platforms = g.Select(i => i.Platform).Distinct().ToList(),
++					}));
++
++				// The inner Targets object is needed for schema compatibility with the Query Mode API.
++				var Result = new
++				{
++					Targets = AggregatedProjectInfo.ToDictionary(item => item.Key, item => item.Target)
++				};
++
++				using FileStream Stream = new(PrimaryProjectFile.FullName, FileMode.Create, FileAccess.Write);
++				JsonSerializer.Serialize(Stream, Result, new JsonSerializerOptions
++				{
++					PropertyNamingPolicy = JsonNamingPolicy.CamelCase,
++					WriteIndented = Minimize == JsonWriterStyle.Readable,
++				});
++			}
++			catch (Exception Ex)
++			{
++				Logger.LogWarning("Exception while writing root project file: {Ex}", Ex.ToString());
++				return false;
++			}
++
++			return true;
++		}
++
++		/// <inheritdoc />
++		protected override FileReference GetProjectLocation(string BaseName)
++		{
++			return FileReference.Combine(IntermediateProjectFilesPath, ProjectFilesFolder, BaseName + ProjectFileExtension);
++		}
++	}
++}
+diff --git a/Engine/Source/Programs/UnrealBuildTool/ToolChain/UEToolChain.cs b/Engine/Source/Programs/UnrealBuildTool/ToolChain/UEToolChain.cs
+index 3c0ead98e8d3f..c2b0a0d0b41a4 100644
+--- a/Engine/Source/Programs/UnrealBuildTool/ToolChain/UEToolChain.cs
++++ b/Engine/Source/Programs/UnrealBuildTool/ToolChain/UEToolChain.cs
+@@ -35,6 +35,12 @@ public static DirectoryReference GetModuleInterfaceDir(DirectoryReference Output
+ 			return DirectoryReference.Combine(OutputDir, "Ifc");
+ 		}
+ 
++		// Return the path to the cpp compiler that will be used by this toolchain.
++		public virtual FileReference? GetCppCompilerPath()
++		{
++			return null;
++		}
++
+ 		public abstract CPPOutput CompileCPPFiles(CppCompileEnvironment CompileEnvironment, List<FileItem> InputFiles, DirectoryReference OutputDir, string ModuleName, IActionGraphBuilder Graph);
+ 
+ 		public virtual CPPOutput CompileRCFiles(CppCompileEnvironment Environment, List<FileItem> InputFiles, DirectoryReference OutputDir, IActionGraphBuilder Graph)
+@@ -67,6 +73,21 @@ public virtual FileItem[] LinkAllFiles(LinkEnvironment LinkEnvironment, bool bBu
+ 			return LinkFile != null ? new FileItem[] { LinkFile } : new FileItem[] { };
+ 		}
+ 
++		public virtual IEnumerable<string> GetGlobalCommandLineArgs(CppCompileEnvironment CompileEnvironment)
++		{
++			return Array.Empty<string>();
++		}
++
++		public virtual IEnumerable<string> GetCPPCommandLineArgs(CppCompileEnvironment CompileEnvironment)
++		{
++			return Array.Empty<string>();
++		}
++
++		public virtual IEnumerable<string> GetCCommandLineArgs(CppCompileEnvironment CompileEnvironment)
++		{
++			return Array.Empty<string>();
++		}
++
+ 		/// <summary>
+ 		/// Get the name of the response file for the current linker environment and output file
+ 		/// </summary>
diff --git a/Engine/Extras/VisualStudioWorkspace/UnrealBuildTool-5.2.patch b/Engine/Extras/VisualStudioWorkspace/UnrealBuildTool-5.2.patch
new file mode 100644
--- /dev/null
+++ b/Engine/Extras/VisualStudioWorkspace/UnrealBuildTool-5.2.patch
@@ -0,0 +1,720 @@
+From 377e0a81ea5d312c6703fc8728c9f61a30805307 Mon Sep 17 00:00:00 2001
+From: [REDACTED_PERSON] <[REDACTED_EMAIL]>
+Date: Wed, 17 Jan 2024 19:49:54 -0800
+Subject: [PATCH] [UE5.2] Project Generator for VS Workspace support
+
+Adding a new Project Generator for Visual Studio that will be used for
+the new support to opening `.uproject` files directly in the IDE without
+generating the traditional `.sln` and .vcxproj files.
+
+This is based on the QueryMode API in `ue5-main` and adapted to the
+GenerateProjectFiles flow. The project files will contain information for each
+Target+Configuration+Platform combinations that are used the project.
+
+Visual Studio will automatically invoke this new generator during the flow to
+open a `.uproject` file directly.
+---
+ .../Modes/GenerateProjectFilesMode.cs         |   4 +
+ .../Platform/Windows/VCToolChain.cs           |  21 ++
+ .../ProjectFiles/ProjectFileGenerator.cs      |   5 +-
+ .../VSWorkspaceProjectFile.cs                 | 314 ++++++++++++++++++
+ .../VSWorkspaceProjectFileGenerator.cs        | 258 ++++++++++++++
+ .../UnrealBuildTool/ToolChain/UEToolChain.cs  |  15 +
+ 6 files changed, 615 insertions(+), 2 deletions(-)
+ create mode 100644 Engine/Source/Programs/UnrealBuildTool/ProjectFiles/VisualStudioWorkspace/VSWorkspaceProjectFile.cs
+ create mode 100644 Engine/Source/Programs/UnrealBuildTool/ProjectFiles/VisualStudioWorkspace/VSWorkspaceProjectFileGenerator.cs
+
+diff --git a/Engine/Source/Programs/UnrealBuildTool/Modes/GenerateProjectFilesMode.cs b/Engine/Source/Programs/UnrealBuildTool/Modes/GenerateProjectFilesMode.cs
+index bbaa8dd2c7cc5..21ee5a01f1bec 100644
+--- a/Engine/Source/Programs/UnrealBuildTool/Modes/GenerateProjectFilesMode.cs
++++ b/Engine/Source/Programs/UnrealBuildTool/Modes/GenerateProjectFilesMode.cs
+@@ -36,6 +36,7 @@ class GenerateProjectFilesMode : ToolMode
+ 		[CommandLine("-EddieProjectFiles", Value = nameof(ProjectFileFormat.Eddie))]
+ 		[CommandLine("-VSCode", Value = nameof(ProjectFileFormat.VisualStudioCode))]
+ 		[CommandLine("-VSMac", Value = nameof(ProjectFileFormat.VisualStudioMac))]
++		[CommandLine("-VSWorkspace", Value = nameof(ProjectFileFormat.VisualStudioWorkspace))]
+ 		[CommandLine("-CLion", Value = nameof(ProjectFileFormat.CLion))]
+ 		[CommandLine("-Rider", Value = nameof(ProjectFileFormat.Rider))]
+ 		#if __VPROJECT_AVAILABLE__
+@@ -210,6 +211,9 @@ public override int Execute(CommandLineArguments Arguments, ILogger Logger)
+ 					case ProjectFileFormat.VisualStudioCode:
+ 						Generator = new VSCodeProjectFileGenerator(ProjectFile);
+ 						break;
++					case ProjectFileFormat.VisualStudioWorkspace:
++						Generator = new VSWorkspaceProjectFileGenerator(ProjectFile, Arguments);
++						break;
+ 					case ProjectFileFormat.CLion:
+ 						Generator = new CLionGenerator(ProjectFile);
+ 						break;
+diff --git a/Engine/Source/Programs/UnrealBuildTool/Platform/Windows/VCToolChain.cs b/Engine/Source/Programs/UnrealBuildTool/Platform/Windows/VCToolChain.cs
+index a3dbb7ff8016d..4327222e6375d 100644
+--- a/Engine/Source/Programs/UnrealBuildTool/Platform/Windows/VCToolChain.cs
++++ b/Engine/Source/Programs/UnrealBuildTool/Platform/Windows/VCToolChain.cs
+@@ -268,6 +268,27 @@ public static void AddSourceDependsFile(List<string> Arguments, FileItem SourceD
+ 			Arguments.Add($"/clang:-MD /clang:-MF\"{SourceDependsFileString}\"");
+ 		}
+ 
++		public override IEnumerable<string> GetGlobalCommandLineArgs(CppCompileEnvironment CompileEnvironment)
++		{
++			List<string> Arguments = new();
++			AppendCLArguments_Global(new(CompileEnvironment), Arguments);
++			return Arguments;
++		}
++
++		public override IEnumerable<string> GetCPPCommandLineArgs(CppCompileEnvironment CompileEnvironment)
++		{
++			List<string> Arguments = new();
++			AppendCLArguments_CPP(CompileEnvironment, Arguments);
++			return Arguments;
++		}
++
++		public override IEnumerable<string> GetCCommandLineArgs(CppCompileEnvironment CompileEnvironment)
++		{
++			List<string> Arguments = new();
++			AppendCLArguments_C(CompileEnvironment, Arguments);
++			return Arguments;
++		}
++
+ 		protected virtual void AppendCLArguments_Global(CppCompileEnvironment CompileEnvironment, List<string> Arguments)
+ 		{
+ 			// Suppress generation of object code for unreferenced inline functions. Enabling this option is more standards compliant, and causes a big reduction
+diff --git a/Engine/Source/Programs/UnrealBuildTool/ProjectFiles/ProjectFileGenerator.cs b/Engine/Source/Programs/UnrealBuildTool/ProjectFiles/ProjectFileGenerator.cs
+index 5e2b1879b8275..0909a2276b066 100644
+--- a/Engine/Source/Programs/UnrealBuildTool/ProjectFiles/ProjectFileGenerator.cs
++++ b/Engine/Source/Programs/UnrealBuildTool/ProjectFiles/ProjectFileGenerator.cs
+@@ -133,6 +133,7 @@ enum ProjectFileFormat
+ 		VisualStudio,
+ 		VisualStudio2019,
+ 		VisualStudio2022,
++		VisualStudioWorkspace,
+ 		XCode,
+ 		Eddie,
+ 		VisualStudioCode,
+@@ -2368,7 +2369,7 @@ private ProjectFile FindProjectForModule(FileReference CurModuleFile, List<FileR
+ 		/// <param name="ProgramProjects">Map of program names to all of the program projects we created</param>
+ 		/// <param name="RulesAssemblies">Map of RuleAssemblies to their base folders</param>
+ 		/// <param name="Logger">Logger for output</param>
+-		private void AddProjectsForAllTargets(
++		protected void AddProjectsForAllTargets(
+ 			PlatformProjectGeneratorCollection PlatformProjectGenerators,
+ 			List<FileReference> AllGames,
+ 			List<FileReference> AllTargetFiles,
+@@ -2665,7 +2666,7 @@ protected void AddProjectsForMods(List<ProjectFile> GameProjects, List<ProjectFi
+ 		/// </summary>
+ 		/// <param name="BaseName">The base name for the project file</param>
+ 		/// <returns>Full path to the project file</returns>
+-		protected FileReference GetProjectLocation(string BaseName)
++		protected virtual FileReference GetProjectLocation(string BaseName)
+ 		{
+ 			return FileReference.Combine(IntermediateProjectFilesPath, BaseName + ProjectFileExtension);
+ 		}
+diff --git a/Engine/Source/Programs/UnrealBuildTool/ProjectFiles/VisualStudioWorkspace/VSWorkspaceProjectFile.cs b/Engine/Source/Programs/UnrealBuildTool/ProjectFiles/VisualStudioWorkspace/VSWorkspaceProjectFile.cs
+new file mode 100644
+index 0000000000000..5e30b56949ab6
+--- /dev/null
++++ b/Engine/Source/Programs/UnrealBuildTool/ProjectFiles/VisualStudioWorkspace/VSWorkspaceProjectFile.cs
+@@ -0,0 +1,314 @@
++// Copyright Epic Games, Inc. All Rights Reserved.
++
++using System;
++using System.Collections.Generic;
++using System.IO;
++using System.Linq;
++using System.Text;
++using System.Text.Json;
++using EpicGames.Core;
++using Microsoft.Extensions.Logging;
++using UnrealBuildBase;
++
++namespace UnrealBuildTool
++{
++	internal class VSWorkspaceProjectFile : ProjectFile
++	{
++		private readonly DirectoryReference RootPath;
++		private readonly HashSet<TargetType> TargetTypes;
++		private readonly CommandLineArguments Arguments;
++
++		/// <summary>
++		/// Collection of output files for this project
++		/// </summary>
++		public List<ExportedTargetInfo> ExportedTargetProjects { get; set; } = new();
++
++		public VSWorkspaceProjectFile(FileReference InProjectFilePath, DirectoryReference BaseDir,
++			DirectoryReference RootPath, HashSet<TargetType> TargetTypes, CommandLineArguments Arguments)
++			: base(InProjectFilePath, BaseDir)
++		{
++			this.RootPath = RootPath;
++			this.TargetTypes = TargetTypes;
++			this.Arguments = Arguments;
++		}
++
++		/// <summary>
++		/// Write project file info in JSON file.
++		/// For every combination of <c>UnrealTargetPlatform</c>, <c>UnrealTargetConfiguration</c> and <c>TargetType</c>
++		/// will be generated separate JSON file.
++		/// </summary>
++		public bool WriteProjectFile(List<UnrealTargetPlatform> InPlatforms,
++			List<UnrealTargetConfiguration> InConfigurations,
++			PlatformProjectGeneratorCollection PlatformProjectGenerators, JsonWriterStyle Minimize, ILogger Logger)
++		{
++			DirectoryReference ProjectRootFolder = RootPath;
++
++			Dictionary<FileReference, (UEBuildTarget BuildTarget, bool bBuildByDefault)> FileToTarget = new();
++
++			foreach (UnrealTargetPlatform Platform in InPlatforms)
++			{
++				foreach (UnrealTargetConfiguration Configuration in InConfigurations)
++				{
++					foreach (ProjectTarget ProjectTarget in ProjectTargets.OfType<ProjectTarget>())
++					{
++						if (TargetTypes.Any() && !TargetTypes.Contains(ProjectTarget.TargetRules!.Type))
++						{
++							continue;
++						}
++
++						// Skip Programs for all configs except for current platform + Development & Debug configurations
++						if (ProjectTarget.TargetRules!.Type == TargetType.Program &&
++							(BuildHostPlatform.Current.Platform != Platform ||
++							 !(Configuration == UnrealTargetConfiguration.Development || Configuration == UnrealTargetConfiguration.Debug)))
++						{
++							continue;
++						}
++
++						// Skip Editor for all platforms except for current platform
++						if (ProjectTarget.TargetRules.Type == TargetType.Editor &&
++							(BuildHostPlatform.Current.Platform != Platform ||
++							(Configuration == UnrealTargetConfiguration.Test || Configuration == UnrealTargetConfiguration.Shipping)))
++						{
++							continue;
++						}
++
++						bool bBuildByDefault = ShouldBuildByDefaultForSolutionTargets && ProjectTarget.SupportedPlatforms.Contains(Platform);
++
++						UnrealArchitectures ProjectArchitectures = UEBuildPlatform
++							.GetBuildPlatform(Platform)
++							.ArchitectureConfig.ActiveArchitectures(ProjectTarget.UnrealProjectFilePath, ProjectTarget.Name);
++
++						TargetDescriptor TargetDesc = new(ProjectTarget.UnrealProjectFilePath, ProjectTarget.Name,
++							Platform, Configuration, ProjectArchitectures, Arguments);
++
++						try
++						{
++							FileReference OutputFile = FileReference.Combine(ProjectRootFolder,
++								$"{ProjectTarget.TargetFilePath.GetFileNameWithoutAnyExtensions()}_{Configuration}_{Platform}.json");
++
++							UEBuildTarget BuildTarget = UEBuildTarget.Create(TargetDesc, false, false, false, Logger);
++							FileToTarget.Add(OutputFile, (BuildTarget, bBuildByDefault));
++						}
++						catch (Exception Ex)
++						{
++							Logger.LogWarning("Exception while generating include data for Target:{Target}, Platform: {Platform}, Configuration: {Configuration}", TargetDesc.Name, Platform.ToString(), Configuration.ToString());
++							Logger.LogWarning("{Ex}", Ex.ToString());
++						}
++					}
++				}
++			}
++
++			foreach (var Entry in FileToTarget)
++			{
++				var OutputFile = Entry.Key;
++				var BuildTarget = Entry.Value.BuildTarget;
++				var bBuildByDefault = Entry.Value.bBuildByDefault;
++
++				try
++				{
++					BuildTarget.PreBuildSetup(Logger);
++
++					ExportedTargetInfo TargetInfo = ExportTarget(BuildTarget, bBuildByDefault, PlatformProjectGenerators, Logger);
++
++					DirectoryReference.CreateDirectory(OutputFile.Directory);
++					using FileStream Stream = new(OutputFile.FullName, FileMode.Create, FileAccess.Write);
++					JsonSerializer.Serialize(Stream, TargetInfo, options: new JsonSerializerOptions()
++					{
++						PropertyNamingPolicy = JsonNamingPolicy.CamelCase,
++						WriteIndented = Minimize == JsonWriterStyle.Readable,
++					});
++
++					ExportedTargetProjects.Add(TargetInfo);
++				}
++				catch (Exception Ex)
++				{
++					Logger.LogWarning("Exception while generating include data for Target:{Target}, Platform: {Platform}, Configuration: {Configuration}",
++						BuildTarget.AppName, BuildTarget.Platform.ToString(), BuildTarget.Configuration.ToString());
++					Logger.LogWarning("{Ex}", Ex.ToString());
++				}
++			}
++
++			return true;
++		}
++
++		private ExportedTargetInfo ExportTarget(
++			UEBuildTarget Target, bool bBuildByDefault, PlatformProjectGeneratorCollection PlatformProjectGenerators, ILogger Logger)
++		{
++			ExportedTargetInfo TargetInfo = new()
++			{
++				TargetName = Target.TargetName,
++				TargetPath = Target.TargetRulesFile.FullName,
++				ProjectPath = Target.ProjectFile?.FullName ?? String.Empty,
++				TargetType = Target.TargetType.ToString(),
++				Platform = Target.Platform.ToString(),
++				Configuration = Target.Configuration.ToString(),
++				BuildInfo = ExportBuildInfo(Target, PlatformProjectGenerators, bBuildByDefault, Logger)
++			};
++
++			UEToolChain TargetToolChain = Target.CreateToolchain(Target.Platform);
++			CppCompileEnvironment GlobalCompileEnvironment = Target.CreateCompileEnvironmentForProjectFiles(Logger);
++
++			HashSet<string> ModuleNames = new();
++			foreach (UEBuildBinary Binary in Target.Binaries)
++			{
++				CppCompileEnvironment BinaryCompileEnvironment = Binary.CreateBinaryCompileEnvironment(GlobalCompileEnvironment);
++				IEnumerable<UEBuildModuleCPP> CandidateModules = Binary.Modules.Where(x => x is UEBuildModuleCPP).Cast<UEBuildModuleCPP>();
++
++				foreach (var ModuleCpp in CandidateModules)
++				{
++					if (!ModuleNames.Add(ModuleCpp.Name))
++					{
++						continue;
++					}
++
++					CppCompileEnvironment ModuleCompileEnvironment = ModuleCpp.CreateCompileEnvironmentForIntellisense(Target.Rules, BinaryCompileEnvironment, Logger);
++					TargetInfo.ModuleToCompileSettings.Add(ModuleCpp.Name, ExportModule(ModuleCpp, TargetToolChain, ModuleCompileEnvironment, Logger));
++
++					foreach (DirectoryReference ModuleDirectory in ModuleCpp.ModuleDirectories)
++					{
++						TargetInfo.DirToModule.TryAdd(ModuleDirectory.FullName, ModuleCpp.Name);
++					}
++
++					if (ModuleCpp.GeneratedCodeDirectory != null)
++					{
++						TargetInfo.DirToModule.TryAdd(ModuleCpp.GeneratedCodeDirectory.FullName, ModuleCpp.Name);
++					}
++				}
++			}
++
++			return TargetInfo;
++		}
++
++		private static ExportedModuleInfo ExportModule(UEBuildModuleCPP Module, UEToolChain TargetToolChain, CppCompileEnvironment ModuleCompileEnvironment, ILogger Logger)
++		{
++			ExportedModuleInfo Result = new()
++			{
++				Name = Module.Name,
++				Directory = Module.ModuleDirectory.FullName,
++				Rules = Module.RulesFile.FullName,
++				GeneratedCodeDirectory = Module.GeneratedCodeDirectory != null ? Module.GeneratedCodeDirectory.FullName : String.Empty,
++				Standard = ModuleCompileEnvironment.CppStandard.ToString(),
++			};
++
++			Result.IncludePaths.AddRange(Module.PublicIncludePaths.Select(x => x.FullName));
++			Result.IncludePaths.AddRange(Module.PublicSystemIncludePaths.Select(x => x.FullName));
++			Result.IncludePaths.AddRange(Module.InternalIncludePaths.Select(x => x.FullName));
++			Result.IncludePaths.AddRange(Module.LegacyPublicIncludePaths.Select(x => x.FullName));
++			Result.IncludePaths.AddRange(Module.LegacyParentIncludePaths.Select(x => x.FullName));
++			Result.IncludePaths.AddRange(Module.PrivateIncludePaths.Select(x => x.FullName));
++
++			Result.IncludePaths.AddRange(ModuleCompileEnvironment.UserIncludePaths.Select(x => x.FullName));
++
++			Result.IncludePaths.AddRange(Module.PublicSystemLibraryPaths.Select(x => x.FullName));
++			Result.IncludePaths.AddRange(Module.PublicSystemLibraries.Concat(Module.PublicLibraries.Select(x => x.FullName)));
++
++			if (TargetToolChain is VCToolChain TargetVCToolChain
++				&& OperatingSystem.IsWindows())
++			{
++				string VCIncludePaths = VCToolChain.GetVCIncludePaths(ModuleCompileEnvironment.Platform, WindowsCompiler.VisualStudio2022, null, Logger);
++				Result.IncludePaths.AddRange(VCIncludePaths.Split(";"));
++			}
++
++			Result.Defines.AddRange(Module.PublicDefinitions);
++			Result.Defines.AddRange(Module.Rules.PrivateDefinitions);
++			Result.Defines.AddRange(Module.Rules.bTreatAsEngineModule ? Array.Empty<string>() : Module.Rules.Target.ProjectDefinitions);
++			Result.Defines.AddRange(Module.GetEmptyApiMacros());
++
++			var ForcedIncludes = ModuleCompileEnvironment.ForceIncludeFiles.ToList();
++			if (ModuleCompileEnvironment.PrecompiledHeaderAction == PrecompiledHeaderAction.Include)
++			{
++				FileItem IncludeHeader = FileItem.GetItemByFileReference(ModuleCompileEnvironment.PrecompiledHeaderIncludeFilename!);
++				ForcedIncludes.Insert(0, IncludeHeader);
++			}
++
++			Result.ForcedIncludes.AddRange(ForcedIncludes.Select(x => x.FullName));
++
++			Result.CompilerPath = TargetToolChain.GetCppCompilerPath()?.ToString();
++			Result.CompilerArgs.AddRange(TargetToolChain.GetGlobalCommandLineArgs(ModuleCompileEnvironment));
++			Result.CompilerAdditionalArgs.Add("c", TargetToolChain.GetCCommandLineArgs(ModuleCompileEnvironment).ToList());
++			Result.CompilerAdditionalArgs.Add("cpp", TargetToolChain.GetCPPCommandLineArgs(ModuleCompileEnvironment).ToList());
++
++			return Result;
++		}
++
++		private TargetBuildInfo? ExportBuildInfo(UEBuildTarget Target, PlatformProjectGeneratorCollection PlatformProjectGenerators, bool bBuildByDefault, ILogger Logger)
++		{
++			if (!BuildHostPlatform.Current.Platform.IsInGroup(UnrealPlatformGroup.Windows))
++			{
++				Logger.LogWarning("Unsupported platform for Build Information: {Platform}", BuildHostPlatform.Current.Platform.ToString());
++				return null;
++			}
++
++			if (IsStubProject)
++			{
++				return null;
++			}
++
++			ProjectTarget ProjectTarget = ProjectTargets.OfType<ProjectTarget>().Single(It => Target.TargetRulesFile == It.TargetFilePath);
++			UnrealTargetPlatform Platform = Target.Platform;
++			UnrealTargetConfiguration Configuration = Target.Configuration;
++
++			string UProjectPath = "";
++			if (IsForeignProject)
++			{
++				UProjectPath = String.Format("\"{0}\"", ProjectTarget.UnrealProjectFilePath!.FullName);
++			}
++
++			PlatformProjectGenerator? ProjGenerator = PlatformProjectGenerators.GetPlatformProjectGenerator(Platform, true);
++			VCProjectFile.BuildCommandBuilder BuildCommandBuilder = new(Configuration, Platform, ProjectTarget, UProjectPath)
++			{
++				ProjectGenerator = ProjGenerator,
++				bIsForeignProject = IsForeignProject
++			};
++
++			string BuildArguments = BuildCommandBuilder.GetBuildArguments();
++
++			return new TargetBuildInfo()
++			{
++				BuildCmd = $"{BuildCommandBuilder.BuildScript.FullName} {BuildArguments}",
++				RebuildCmd = $"{BuildCommandBuilder.RebuildScript.FullName} {BuildArguments}",
++				CleanCmd = $"{BuildCommandBuilder.CleanScript.FullName} {BuildArguments}",
++				PrimaryOutput = Target.Binaries[0].OutputFilePath.FullName,
++				BuildByDefault = bBuildByDefault,
++			};
++		}
++
++		internal class ExportedTargetInfo
++		{
++			public string TargetName { get; set; } = String.Empty;
++			public string TargetPath { get; set; } = String.Empty;
++			public string ProjectPath { get; set; } = String.Empty;
++			public string TargetType { get; set; } = String.Empty;
++			public string Platform { get; set; } = String.Empty;
++			public string Configuration { get; set; } = String.Empty;
++			public TargetBuildInfo? BuildInfo { get; set; }
++			public Dictionary<string, ExportedModuleInfo> ModuleToCompileSettings { get; set; } = new();
++			public Dictionary<string, string> DirToModule { get; set; } = new();
++		}
++
++		internal class ExportedModuleInfo
++		{
++			public string Name { get; set; } = String.Empty;
++			public string Directory { get; set; } = String.Empty;
++			public string Rules { get; set; } = String.Empty;
++			public string GeneratedCodeDirectory { get; set; } = String.Empty;
++			public List<string> IncludePaths { get; set; } = new();
++			public List<string> Defines { get; set; } = new();
++			public string? Standard { get; set; }
++			public List<string> ForcedIncludes { get; set; } = new();
++			public string? CompilerPath { get; set; }
++			public List<string> CompilerArgs { get; set; } = new();
++			public Dictionary<string, List<string>> CompilerAdditionalArgs { get; set; } = new();
++			public string? WindowsSdkVersion { get; set; }
++		}
++
++		internal class TargetBuildInfo
++		{
++			public string BuildCmd { get; set; } = String.Empty;
++			public string RebuildCmd { get; set; } = String.Empty;
++			public string CleanCmd { get; set; } = String.Empty;
++			public string PrimaryOutput { get; set; } = String.Empty;
++			public bool BuildByDefault { get; internal set; }
++		}
++	}
++}
+diff --git a/Engine/Source/Programs/UnrealBuildTool/ProjectFiles/VisualStudioWorkspace/VSWorkspaceProjectFileGenerator.cs b/Engine/Source/Programs/UnrealBuildTool/ProjectFiles/VisualStudioWorkspace/VSWorkspaceProjectFileGenerator.cs
+new file mode 100644
+index 0000000000000..25093e8a4b852
+--- /dev/null
++++ b/Engine/Source/Programs/UnrealBuildTool/ProjectFiles/VisualStudioWorkspace/VSWorkspaceProjectFileGenerator.cs
+@@ -0,0 +1,258 @@
++// Copyright Epic Games, Inc. All Rights Reserved.
++
++using System;
++using System.Collections.Generic;
++using System.IO;
++using System.Linq;
++using System.Text.Json;
++using EpicGames.Core;
++using Microsoft.Extensions.Logging;
++using UnrealBuildBase;
++
++namespace UnrealBuildTool
++{
++	class VSWorkspaceProjectFileGenerator : ProjectFileGenerator
++	{
++		public override string ProjectFileExtension => ".json";
++
++		// These properties are used by Visual Studio to determine where to read the project files.
++		// So they must remain constant.
++		private const string VSUnrealWorkspaceFileName = ".vs-unreal-workspace";
++		private const string ProjectFilesFolder = "VisualStudio";
++
++		private readonly CommandLineArguments Arguments;
++
++		/// <summary>
++		/// List of deprecated platforms.
++		/// Don't generate project model for these platforms unless they are specified in "Platforms" console arguments.
++		/// </summary>
++		/// <returns></returns>
++		private readonly HashSet<UnrealTargetPlatform> DeprecatedPlatforms = new();
++
++		/// <summary>
++		/// Platforms to generate project files for
++		/// </summary>
++		[CommandLine("-Platforms=", ListSeparator = '+')]
++		HashSet<UnrealTargetPlatform> Platforms = new();
++
++		/// <summary>
++		/// Target types to generate project files for
++		/// </summary>
++		[CommandLine("-TargetTypes=", ListSeparator = '+')]
++		HashSet<TargetType> TargetTypes = new();
++
++		/// <summary>
++		/// Target configurations to generate project files for
++		/// </summary>
++		[CommandLine("-TargetConfigurations=", ListSeparator = '+')]
++		HashSet<UnrealTargetConfiguration> TargetConfigurations = new();
++
++		/// <summary>
++		/// Projects to generate project files for
++		/// </summary>
++		[CommandLine("-ProjectNames=", ListSeparator = '+')]
++		HashSet<string> ProjectNames = new();
++
++		/// <summary>
++		/// Should format JSON files in human readable form, or use packed one without indents
++		/// </summary>
++		[CommandLine("-Minimize", Value = "Compact")]
++		private JsonWriterStyle Minimize = JsonWriterStyle.Readable;
++
++		public VSWorkspaceProjectFileGenerator(FileReference? InOnlyGameProject,
++			CommandLineArguments InArguments)
++			: base(InOnlyGameProject)
++		{
++			Arguments = InArguments;
++			Arguments.ApplyTo(this);
++		}
++
++		public override bool ShouldGenerateIntelliSenseData() => true;
++
++		public override void CleanProjectFiles(DirectoryReference InPrimaryProjectDirectory, string InPrimaryProjectName,
++			DirectoryReference InIntermediateProjectFilesDirectory, ILogger Logger)
++		{
++			DirectoryReference.Delete(InPrimaryProjectDirectory);
++		}
++
++		protected override void ConfigureProjectFileGeneration(string[] Arguments, ref bool IncludeAllPlatforms, ILogger Logger)
++		{
++			base.ConfigureProjectFileGeneration(Arguments, ref IncludeAllPlatforms, Logger);
++		}
++
++		protected override ProjectFile AllocateProjectFile(FileReference InitFilePath, DirectoryReference BaseDir)
++		{
++			VSWorkspaceProjectFile projectFile = new(InitFilePath, BaseDir, RootPath: InitFilePath.Directory,
++				Arguments: Arguments, TargetTypes: TargetTypes);
++			return projectFile;
++		}
++
++		protected override bool WriteProjectFiles(PlatformProjectGeneratorCollection PlatformProjectGenerators, ILogger Logger)
++		{
++			using ProgressWriter Progress = new("Writing project files...", true, Logger);
++			List<ProjectFile> ProjectsToGenerate = new(GeneratedProjectFiles);
++			if (ProjectNames.Any())
++			{
++				ProjectsToGenerate = ProjectsToGenerate.Where(Project =>
++					ProjectNames.Contains(Project.ProjectFilePath.GetFileNameWithoutAnyExtensions())).ToList();
++			}
++
++			int TotalProjectFileCount = ProjectsToGenerate.Count;
++
++			HashSet<UnrealTargetPlatform> PlatformsToGenerate = new(SupportedPlatforms);
++			if (Platforms.Any())
++			{
++				PlatformsToGenerate.IntersectWith(Platforms);
++			}
++
++			List<UnrealTargetPlatform> FilteredPlatforms = PlatformsToGenerate.Where(Platform =>
++			{
++				// Skip deprecated unless explicitly specified in the command line.
++				return (!DeprecatedPlatforms.Contains(Platform) || Platforms.Contains(Platform))
++					   && UEBuildPlatform.IsPlatformAvailable(Platform);
++			}).ToList();
++
++			HashSet<UnrealTargetConfiguration> ConfigurationsToGenerate = new(SupportedConfigurations);
++			if (TargetConfigurations.Any())
++			{
++				ConfigurationsToGenerate.IntersectWith(TargetConfigurations);
++			}
++
++			List<UnrealTargetConfiguration> Configurations = ConfigurationsToGenerate.ToList();
++
++			for (int ProjectFileIndex = 0; ProjectFileIndex < ProjectsToGenerate.Count; ++ProjectFileIndex)
++			{
++				if (ProjectsToGenerate[ProjectFileIndex] is not VSWorkspaceProjectFile CurrentProject)
++				{
++					return false;
++				}
++
++				if (!CurrentProject.WriteProjectFile(FilteredPlatforms, Configurations, PlatformProjectGenerators, Minimize, Logger))
++				{
++					return false;
++				}
++
++				Progress.Write(ProjectFileIndex + 1, TotalProjectFileCount);
++			}
++
++			Progress.Write(TotalProjectFileCount, TotalProjectFileCount);
++
++			return true;
++		}
++
++		public override bool GenerateProjectFiles(PlatformProjectGeneratorCollection PlatformProjectGenerators,
++			String[] Arguments, ILogger Logger)
++		{
++			bool IncludeAllPlatforms = true;
++			ConfigureProjectFileGeneration(Arguments, ref IncludeAllPlatforms, Logger);
++
++			if (bGeneratingGameProjectFiles || Unreal.IsEngineInstalled())
++			{
++				PrimaryProjectPath = OnlyGameProject!.Directory;
++				PrimaryProjectName = OnlyGameProject.GetFileNameWithoutExtension();
++
++				IntermediateProjectFilesPath =
++					DirectoryReference.Combine(PrimaryProjectPath, "Intermediate", "ProjectFiles");
++			}
++
++			SetupSupportedPlatformsAndConfigurations(IncludeAllPlatforms: true, Logger, out string SupportedPlatformNames);
++			Logger.LogDebug("Supported platforms: {Platforms}", SupportedPlatformNames);
++
++			List<FileReference> AllGames = FindGameProjects(Logger);
++
++			{
++				// Find all of the target files.
++				List<FileReference> AllTargetFiles = DiscoverTargets(
++					AllGames,
++					Logger,
++					OnlyGameProject,
++					SupportedPlatforms,
++					bIncludeEngineSource: bIncludeEngineSource,
++					bIncludeTempTargets: bIncludeTempTargets);
++
++				// If there are multiple targets of a given type for a project, use the order to determine which one gets a suffix.
++				AllTargetFiles = AllTargetFiles.OrderBy(x => x.FullName, StringComparer.OrdinalIgnoreCase).ToList();
++
++				ProjectFile? EngineProject = null;
++				List<ProjectFile> GameProjects = new();
++				List<ProjectFile> ModProjects = new();
++				Dictionary<FileReference, ProjectFile> ProgramProjects = new();
++				Dictionary<RulesAssembly, DirectoryReference> RulesAssemblies = new();
++				Dictionary<ProjectFile, FileReference> ProjectFileToUProjectFile = new();
++
++				AddProjectsForAllTargets(
++					PlatformProjectGenerators,
++					AllGames,
++					AllTargetFiles,
++					Arguments,
++					ref EngineProject,
++					GameProjects,
++					ProjectFileToUProjectFile,
++					ProgramProjects,
++					RulesAssemblies,
++					Logger);
++
++				AddAllGameProjects(GameProjects);
++			}
++
++			WriteProjectFiles(PlatformProjectGenerators, Logger);
++			WritePrimaryProjectFile(UBTProject, PlatformProjectGenerators, Logger);
++
++			return true;
++		}
++
++		protected override bool WritePrimaryProjectFile(ProjectFile? UBTProject,
++			PlatformProjectGeneratorCollection PlatformProjectGenerators,
++			ILogger Logger)
++		{
++			try
++			{
++				FileReference PrimaryProjectFile = FileReference.Combine(
++					IntermediateProjectFilesPath, ProjectFilesFolder, VSUnrealWorkspaceFileName);
++
++				DirectoryReference.CreateDirectory(PrimaryProjectFile.Directory);
++
++				// Collect all the resulting project files and aggregate the target-level data
++				var AggregatedProjectInfo = GeneratedProjectFiles
++					.Where(Project => Project is VSWorkspaceProjectFile)
++					.OfType<VSWorkspaceProjectFile>()
++					.SelectMany(Project => Project.ExportedTargetProjects)
++					.GroupBy(TargetProject => TargetProject.TargetName)
++					.Select(g => (g.Key, Target: new
++					{
++						TargetType = g.Select(i => i.TargetType).Distinct().Single(),
++						TargetPath = g.Select(i => i.TargetPath).Distinct().Single(),
++						ProjectPath = g.Select(i => i.ProjectPath).Distinct().Single(),
++						Configurations = g.Select(i => i.Configuration).Distinct().ToList(),
++						Platforms = g.Select(i => i.Platform).Distinct().ToList(),
++					}));
++
++				// The inner Targets object is needed for schema compatibility with the Query Mode API.
++				var Result = new
++				{
++					Targets = AggregatedProjectInfo.ToDictionary(item => item.Key, item => item.Target)
++				};
++
++				using FileStream Stream = new(PrimaryProjectFile.FullName, FileMode.Create, FileAccess.Write);
++				JsonSerializer.Serialize(Stream, Result, new JsonSerializerOptions
++				{
++					PropertyNamingPolicy = JsonNamingPolicy.CamelCase,
++					WriteIndented = Minimize == JsonWriterStyle.Readable,
++				});
++			}
++			catch (Exception Ex)
++			{
++				Logger.LogWarning("Exception while writing root project file: {Ex}", Ex.ToString());
++				return false;
++			}
++
++			return true;
++		}
++
++		/// <inheritdoc />
++		protected override FileReference GetProjectLocation(string BaseName)
++		{
++			return FileReference.Combine(IntermediateProjectFilesPath, ProjectFilesFolder, BaseName + ProjectFileExtension);
++		}
++	}
++}
+diff --git a/Engine/Source/Programs/UnrealBuildTool/ToolChain/UEToolChain.cs b/Engine/Source/Programs/UnrealBuildTool/ToolChain/UEToolChain.cs
+index c9a3d197d078a..05eb8ce606b03 100644
+--- a/Engine/Source/Programs/UnrealBuildTool/ToolChain/UEToolChain.cs
++++ b/Engine/Source/Programs/UnrealBuildTool/ToolChain/UEToolChain.cs
+@@ -149,6 +149,21 @@ public virtual FileItem[] LinkAllFiles(LinkEnvironment LinkEnvironment, bool bBu
+ 			return Result.ToArray();
+ 		}
+ 
++		public virtual IEnumerable<string> GetGlobalCommandLineArgs(CppCompileEnvironment CompileEnvironment)
++		{
++			return Array.Empty<string>();
++		}
++
++		public virtual IEnumerable<string> GetCPPCommandLineArgs(CppCompileEnvironment CompileEnvironment)
++		{
++			return Array.Empty<string>();
++		}
++
++		public virtual IEnumerable<string> GetCCommandLineArgs(CppCompileEnvironment CompileEnvironment)
++		{
++			return Array.Empty<string>();
++		}
++
+ 		/// <summary>
+ 		/// Get the name of the response file for the current linker environment and output file
+ 		/// </summary>
diff --git a/Engine/Extras/VisualStudioWorkspace/UnrealBuildTool-5.3.patch b/Engine/Extras/VisualStudioWorkspace/UnrealBuildTool-5.3.patch
new file mode 100644
--- /dev/null
+++ b/Engine/Extras/VisualStudioWorkspace/UnrealBuildTool-5.3.patch
@@ -0,0 +1,720 @@
+From c263609aa0d477b3010c3a4ea21080e78100f01f Mon Sep 17 00:00:00 2001
+From: [REDACTED_PERSON] <[REDACTED_EMAIL]>
+Date: Wed, 17 Jan 2024 19:49:54 -0800
+Subject: [PATCH] [UE5.3] Project Generator for VS Workspace support
+
+Adding a new Project Generator for Visual Studio that will be used for
+the new support to opening `.uproject` files directly in the IDE without
+generating the traditional `.sln` and .vcxproj files.
+
+This is based on the QueryMode API in `ue5-main` and adapted to the
+GenerateProjectFiles flow. The project files will contain information for each
+Target+Configuration+Platform combinations that are used the project.
+
+Visual Studio will automatically invoke this new generator during the flow to
+open a `.uproject` file directly.
+---
+ .../Modes/GenerateProjectFilesMode.cs         |   4 +
+ .../Platform/Windows/VCToolChain.cs           |  21 ++
+ .../ProjectFiles/ProjectFileGenerator.cs      |   5 +-
+ .../VSWorkspaceProjectFile.cs                 | 314 ++++++++++++++++++
+ .../VSWorkspaceProjectFileGenerator.cs        | 258 ++++++++++++++
+ .../UnrealBuildTool/ToolChain/UEToolChain.cs  |  15 +
+ 6 files changed, 615 insertions(+), 2 deletions(-)
+ create mode 100644 Engine/Source/Programs/UnrealBuildTool/ProjectFiles/VisualStudioWorkspace/VSWorkspaceProjectFile.cs
+ create mode 100644 Engine/Source/Programs/UnrealBuildTool/ProjectFiles/VisualStudioWorkspace/VSWorkspaceProjectFileGenerator.cs
+
+diff --git a/Engine/Source/Programs/UnrealBuildTool/Modes/GenerateProjectFilesMode.cs b/Engine/Source/Programs/UnrealBuildTool/Modes/GenerateProjectFilesMode.cs
+index 021fec3479b44..6b12960e19cef 100644
+--- a/Engine/Source/Programs/UnrealBuildTool/Modes/GenerateProjectFilesMode.cs
++++ b/Engine/Source/Programs/UnrealBuildTool/Modes/GenerateProjectFilesMode.cs
+@@ -34,6 +34,7 @@ class GenerateProjectFilesMode : ToolMode
+ 		[CommandLine("-EddieProjectFiles", Value = nameof(ProjectFileFormat.Eddie))]
+ 		[CommandLine("-VSCode", Value = nameof(ProjectFileFormat.VisualStudioCode))]
+ 		[CommandLine("-VSMac", Value = nameof(ProjectFileFormat.VisualStudioMac))]
++		[CommandLine("-VSWorkspace", Value = nameof(ProjectFileFormat.VisualStudioWorkspace))]
+ 		[CommandLine("-CLion", Value = nameof(ProjectFileFormat.CLion))]
+ 		[CommandLine("-Rider", Value = nameof(ProjectFileFormat.Rider))]
+ #if __VPROJECT_AVAILABLE__
+@@ -220,6 +221,9 @@ public override Task<int> ExecuteAsync(CommandLineArguments Arguments, ILogger L
+ 					case ProjectFileFormat.VisualStudioCode:
+ 						Generator = new VSCodeProjectFileGenerator(ProjectFile);
+ 						break;
++					case ProjectFileFormat.VisualStudioWorkspace:
++						Generator = new VSWorkspaceProjectFileGenerator(ProjectFile, Arguments);
++						break;
+ 					case ProjectFileFormat.CLion:
+ 						Generator = new CLionGenerator(ProjectFile);
+ 						break;
+diff --git a/Engine/Source/Programs/UnrealBuildTool/Platform/Windows/VCToolChain.cs b/Engine/Source/Programs/UnrealBuildTool/Platform/Windows/VCToolChain.cs
+index 080f9f65bfdc6..8e5afe8709c33 100644
+--- a/Engine/Source/Programs/UnrealBuildTool/Platform/Windows/VCToolChain.cs
++++ b/Engine/Source/Programs/UnrealBuildTool/Platform/Windows/VCToolChain.cs
+@@ -288,6 +288,27 @@ public static void AddSourceDependsFile(List<string> Arguments, FileItem SourceD
+ 			Arguments.Add($"/clang:-MD /clang:-MF\"{SourceDependsFileString}\"");
+ 		}
+ 
++		public override IEnumerable<string> GetGlobalCommandLineArgs(CppCompileEnvironment CompileEnvironment)
++		{
++			List<string> Arguments = new();
++			AppendCLArguments_Global(new(CompileEnvironment), Arguments);
++			return Arguments;
++		}
++
++		public override IEnumerable<string> GetCPPCommandLineArgs(CppCompileEnvironment CompileEnvironment)
++		{
++			List<string> Arguments = new();
++			AppendCLArguments_CPP(CompileEnvironment, Arguments);
++			return Arguments;
++		}
++
++		public override IEnumerable<string> GetCCommandLineArgs(CppCompileEnvironment CompileEnvironment)
++		{
++			List<string> Arguments = new();
++			AppendCLArguments_C(CompileEnvironment, Arguments);
++			return Arguments;
++		}
++
+ 		protected virtual void AppendCLArguments_Global(CppCompileEnvironment CompileEnvironment, List<string> Arguments)
+ 		{
+ 			// Suppress generation of object code for unreferenced inline functions. Enabling this option is more standards compliant, and causes a big reduction
+diff --git a/Engine/Source/Programs/UnrealBuildTool/ProjectFiles/ProjectFileGenerator.cs b/Engine/Source/Programs/UnrealBuildTool/ProjectFiles/ProjectFileGenerator.cs
+index 91aabfc073872..d861b8c9f658f 100644
+--- a/Engine/Source/Programs/UnrealBuildTool/ProjectFiles/ProjectFileGenerator.cs
++++ b/Engine/Source/Programs/UnrealBuildTool/ProjectFiles/ProjectFileGenerator.cs
+@@ -143,6 +143,7 @@ enum ProjectFileFormat
+ 		VisualStudio,
+ 		VisualStudio2019,
+ 		VisualStudio2022,
++		VisualStudioWorkspace,
+ 		XCode,
+ 		Eddie,
+ 		VisualStudioCode,
+@@ -2629,7 +2630,7 @@ private ProjectFile FindProjectForModule(FileReference CurModuleFile, TargetType
+ 		/// <param name="ProgramProjects">Map of program names to all of the program projects we created</param>
+ 		/// <param name="RulesAssemblies">Map of RuleAssemblies to their base folders</param>
+ 		/// <param name="Logger">Logger for output</param>
+-		private void AddProjectsForAllTargets(
++		protected void AddProjectsForAllTargets(
+ 			PlatformProjectGeneratorCollection PlatformProjectGenerators,
+ 			List<FileReference> AllGames,
+ 			List<FileReference> AllTargetFiles,
+@@ -2994,7 +2995,7 @@ protected void AddProjectsForMods(List<ProjectFile> GameProjects, List<ProjectFi
+ 		/// </summary>
+ 		/// <param name="BaseName">The base name for the project file</param>
+ 		/// <returns>Full path to the project file</returns>
+-		protected FileReference GetProjectLocation(string BaseName)
++		protected virtual FileReference GetProjectLocation(string BaseName)
+ 		{
+ 			return FileReference.Combine(IntermediateProjectFilesPath, BaseName + ProjectFileExtension);
+ 		}
+diff --git a/Engine/Source/Programs/UnrealBuildTool/ProjectFiles/VisualStudioWorkspace/VSWorkspaceProjectFile.cs b/Engine/Source/Programs/UnrealBuildTool/ProjectFiles/VisualStudioWorkspace/VSWorkspaceProjectFile.cs
+new file mode 100644
+index 0000000000000..6c5779fd5c111
+--- /dev/null
++++ b/Engine/Source/Programs/UnrealBuildTool/ProjectFiles/VisualStudioWorkspace/VSWorkspaceProjectFile.cs
+@@ -0,0 +1,314 @@
++// Copyright Epic Games, Inc. All Rights Reserved.
++
++using System;
++using System.Collections.Generic;
++using System.IO;
++using System.Linq;
++using System.Text;
++using System.Text.Json;
++using EpicGames.Core;
++using Microsoft.Extensions.Logging;
++using UnrealBuildBase;
++
++namespace UnrealBuildTool
++{
++	internal class VSWorkspaceProjectFile : ProjectFile
++	{
++		private readonly DirectoryReference RootPath;
++		private readonly HashSet<TargetType> TargetTypes;
++		private readonly CommandLineArguments Arguments;
++
++		/// <summary>
++		/// Collection of output files for this project
++		/// </summary>
++		public List<ExportedTargetInfo> ExportedTargetProjects { get; set; } = new();
++
++		public VSWorkspaceProjectFile(FileReference InProjectFilePath, DirectoryReference BaseDir,
++			DirectoryReference RootPath, HashSet<TargetType> TargetTypes, CommandLineArguments Arguments)
++			: base(InProjectFilePath, BaseDir)
++		{
++			this.RootPath = RootPath;
++			this.TargetTypes = TargetTypes;
++			this.Arguments = Arguments;
++		}
++
++		/// <summary>
++		/// Write project file info in JSON file.
++		/// For every combination of <c>UnrealTargetPlatform</c>, <c>UnrealTargetConfiguration</c> and <c>TargetType</c>
++		/// will be generated separate JSON file.
++		/// </summary>
++		public bool WriteProjectFile(List<UnrealTargetPlatform> InPlatforms,
++			List<UnrealTargetConfiguration> InConfigurations,
++			PlatformProjectGeneratorCollection PlatformProjectGenerators, JsonWriterStyle Minimize, ILogger Logger)
++		{
++			DirectoryReference ProjectRootFolder = RootPath;
++
++			Dictionary<FileReference, (UEBuildTarget BuildTarget, bool bBuildByDefault)> FileToTarget = new();
++
++			foreach (UnrealTargetPlatform Platform in InPlatforms)
++			{
++				foreach (UnrealTargetConfiguration Configuration in InConfigurations)
++				{
++					foreach (ProjectTarget ProjectTarget in ProjectTargets.OfType<ProjectTarget>())
++					{
++						if (TargetTypes.Any() && !TargetTypes.Contains(ProjectTarget.TargetRules!.Type))
++						{
++							continue;
++						}
++
++						// Skip Programs for all configs except for current platform + Development & Debug configurations
++						if (ProjectTarget.TargetRules!.Type == TargetType.Program &&
++							(BuildHostPlatform.Current.Platform != Platform ||
++							 !(Configuration == UnrealTargetConfiguration.Development || Configuration == UnrealTargetConfiguration.Debug)))
++						{
++							continue;
++						}
++
++						// Skip Editor for all platforms except for current platform
++						if (ProjectTarget.TargetRules.Type == TargetType.Editor && 
++							(BuildHostPlatform.Current.Platform != Platform || 
++							(Configuration == UnrealTargetConfiguration.Test || Configuration == UnrealTargetConfiguration.Shipping)))
++						{
++							continue;
++						}
++
++						bool bBuildByDefault = ShouldBuildByDefaultForSolutionTargets && ProjectTarget.SupportedPlatforms.Contains(Platform);
++
++						UnrealArchitectures ProjectArchitectures = UEBuildPlatform
++							.GetBuildPlatform(Platform)
++							.ArchitectureConfig.ActiveArchitectures(ProjectTarget.UnrealProjectFilePath, ProjectTarget.Name);
++
++						TargetDescriptor TargetDesc = new(ProjectTarget.UnrealProjectFilePath, ProjectTarget.Name,
++							Platform, Configuration, ProjectArchitectures, Arguments);
++
++						try
++						{
++							FileReference OutputFile = FileReference.Combine(ProjectRootFolder, 
++								$"{ProjectTarget.TargetFilePath.GetFileNameWithoutAnyExtensions()}_{Configuration}_{Platform}.json");
++
++							UEBuildTarget BuildTarget = UEBuildTarget.Create(TargetDesc, false, false, false, Logger);
++							FileToTarget.Add(OutputFile, (BuildTarget, bBuildByDefault));
++						}
++						catch (Exception Ex)
++						{
++							Logger.LogWarning("Exception while generating include data for Target:{Target}, Platform: {Platform}, Configuration: {Configuration}", TargetDesc.Name, Platform.ToString(), Configuration.ToString());
++							Logger.LogWarning("{Ex}", Ex.ToString());
++						}
++					}
++				}
++			}
++
++			foreach (var Entry in FileToTarget)
++			{
++				var OutputFile = Entry.Key;
++				var BuildTarget = Entry.Value.BuildTarget;
++				var bBuildByDefault = Entry.Value.bBuildByDefault;
++
++				try
++				{
++					BuildTarget.PreBuildSetup(Logger);
++
++					ExportedTargetInfo TargetInfo = ExportTarget(BuildTarget, bBuildByDefault, PlatformProjectGenerators, Logger);
++
++					DirectoryReference.CreateDirectory(OutputFile.Directory);
++					using FileStream Stream = new(OutputFile.FullName, FileMode.Create, FileAccess.Write);
++					JsonSerializer.Serialize(Stream, TargetInfo, options: new JsonSerializerOptions()
++					{
++						PropertyNamingPolicy = JsonNamingPolicy.CamelCase,
++						WriteIndented = Minimize == JsonWriterStyle.Readable,
++					});
++
++					ExportedTargetProjects.Add(TargetInfo);
++				}
++				catch (Exception Ex)
++				{
++					Logger.LogWarning("Exception while generating include data for Target:{Target}, Platform: {Platform}, Configuration: {Configuration}",
++						BuildTarget.AppName, BuildTarget.Platform.ToString(), BuildTarget.Configuration.ToString());
++					Logger.LogWarning("{Ex}", Ex.ToString());
++				}
++			}
++
++			return true;
++		}
++
++		private ExportedTargetInfo ExportTarget(
++			UEBuildTarget Target, bool bBuildByDefault, PlatformProjectGeneratorCollection PlatformProjectGenerators, ILogger Logger)
++		{
++			ExportedTargetInfo TargetInfo = new()
++			{
++				TargetName = Target.TargetName,
++				TargetPath = Target.TargetRulesFile.FullName,
++				ProjectPath = Target.ProjectFile?.FullName ?? String.Empty,
++				TargetType = Target.TargetType.ToString(),
++				Platform = Target.Platform.ToString(),
++				Configuration = Target.Configuration.ToString(),
++				BuildInfo = ExportBuildInfo(Target, PlatformProjectGenerators, bBuildByDefault, Logger)
++			};
++
++			UEToolChain TargetToolChain = Target.CreateToolchain(Target.Platform);
++			CppCompileEnvironment GlobalCompileEnvironment = Target.CreateCompileEnvironmentForProjectFiles(Logger);
++
++			HashSet<string> ModuleNames = new();
++			foreach (UEBuildBinary Binary in Target.Binaries)
++			{
++				CppCompileEnvironment BinaryCompileEnvironment = Binary.CreateBinaryCompileEnvironment(GlobalCompileEnvironment);
++				IEnumerable<UEBuildModuleCPP> CandidateModules = Binary.Modules.Where(x => x is UEBuildModuleCPP).Cast<UEBuildModuleCPP>();
++
++				foreach (var ModuleCpp in CandidateModules)
++				{
++					if (!ModuleNames.Add(ModuleCpp.Name))
++					{
++						continue;
++					}
++
++					CppCompileEnvironment ModuleCompileEnvironment = ModuleCpp.CreateCompileEnvironmentForIntellisense(Target.Rules, BinaryCompileEnvironment, Logger);
++					TargetInfo.ModuleToCompileSettings.Add(ModuleCpp.Name, ExportModule(ModuleCpp, TargetToolChain, ModuleCompileEnvironment, Logger));
++
++					foreach (DirectoryReference ModuleDirectory in ModuleCpp.ModuleDirectories)
++					{
++						TargetInfo.DirToModule.TryAdd(ModuleDirectory.FullName, ModuleCpp.Name);
++					}
++
++					if (ModuleCpp.GeneratedCodeDirectory != null)
++					{
++						TargetInfo.DirToModule.TryAdd(ModuleCpp.GeneratedCodeDirectory.FullName, ModuleCpp.Name);
++					}
++				}
++			}
++
++			return TargetInfo;
++		}
++
++		private static ExportedModuleInfo ExportModule(UEBuildModuleCPP Module, UEToolChain TargetToolChain, CppCompileEnvironment ModuleCompileEnvironment, ILogger Logger)
++		{
++			ExportedModuleInfo Result = new()
++			{
++				Name = Module.Name,
++				Directory = Module.ModuleDirectory.FullName,
++				Rules = Module.RulesFile.FullName,
++				GeneratedCodeDirectory = Module.GeneratedCodeDirectory != null ? Module.GeneratedCodeDirectory.FullName : String.Empty,
++				Standard = ModuleCompileEnvironment.CppStandard.ToString(),
++			};
++
++			Result.IncludePaths.AddRange(Module.PublicIncludePaths.Select(x => x.FullName));
++			Result.IncludePaths.AddRange(Module.PublicSystemIncludePaths.Select(x => x.FullName));
++			Result.IncludePaths.AddRange(Module.InternalIncludePaths.Select(x => x.FullName));
++			Result.IncludePaths.AddRange(Module.LegacyPublicIncludePaths.Select(x => x.FullName));
++			Result.IncludePaths.AddRange(Module.LegacyParentIncludePaths.Select(x => x.FullName));
++			Result.IncludePaths.AddRange(Module.PrivateIncludePaths.Select(x => x.FullName));
++
++			Result.IncludePaths.AddRange(ModuleCompileEnvironment.UserIncludePaths.Select(x => x.FullName));
++
++			Result.IncludePaths.AddRange(Module.PublicSystemLibraryPaths.Select(x => x.FullName));
++			Result.IncludePaths.AddRange(Module.PublicSystemLibraries.Concat(Module.PublicLibraries.Select(x => x.FullName)));
++
++			if (TargetToolChain is VCToolChain TargetVCToolChain
++				&& OperatingSystem.IsWindows())
++			{
++				string VCIncludePaths = VCToolChain.GetVCIncludePaths(ModuleCompileEnvironment.Platform, WindowsCompiler.VisualStudio2022, null, null, Logger);
++				Result.IncludePaths.AddRange(VCIncludePaths.Split(";"));
++			}
++
++			Result.Defines.AddRange(Module.PublicDefinitions);
++			Result.Defines.AddRange(Module.Rules.PrivateDefinitions);
++			Result.Defines.AddRange(Module.Rules.bTreatAsEngineModule ? Array.Empty<string>() : Module.Rules.Target.ProjectDefinitions);
++			Result.Defines.AddRange(Module.GetEmptyApiMacros());
++
++			var ForcedIncludes = ModuleCompileEnvironment.ForceIncludeFiles.ToList();
++			if (ModuleCompileEnvironment.PrecompiledHeaderAction == PrecompiledHeaderAction.Include)
++			{
++				FileItem IncludeHeader = FileItem.GetItemByFileReference(ModuleCompileEnvironment.PrecompiledHeaderIncludeFilename!);
++				ForcedIncludes.Insert(0, IncludeHeader);
++			}
++
++			Result.ForcedIncludes.AddRange(ForcedIncludes.Select(x => x.FullName));
++
++			Result.CompilerPath = TargetToolChain.GetCppCompilerPath()?.ToString();
++			Result.CompilerArgs.AddRange(TargetToolChain.GetGlobalCommandLineArgs(ModuleCompileEnvironment));
++			Result.CompilerAdditionalArgs.Add("c", TargetToolChain.GetCCommandLineArgs(ModuleCompileEnvironment).ToList());
++			Result.CompilerAdditionalArgs.Add("cpp", TargetToolChain.GetCPPCommandLineArgs(ModuleCompileEnvironment).ToList());
++
++			return Result;
++		}
++
++		private TargetBuildInfo? ExportBuildInfo(UEBuildTarget Target, PlatformProjectGeneratorCollection PlatformProjectGenerators, bool bBuildByDefault, ILogger Logger)
++		{
++			if (!BuildHostPlatform.Current.Platform.IsInGroup(UnrealPlatformGroup.Windows))
++			{
++				Logger.LogWarning("Unsupported platform for Build Information: {Platform}", BuildHostPlatform.Current.Platform.ToString());
++				return null;
++			}
++
++			if (IsStubProject)
++			{
++				return null;
++			}
++
++			ProjectTarget ProjectTarget = ProjectTargets.OfType<ProjectTarget>().Single(It => Target.TargetRulesFile == It.TargetFilePath);
++			UnrealTargetPlatform Platform = Target.Platform;
++			UnrealTargetConfiguration Configuration = Target.Configuration;
++
++			string UProjectPath = "";
++			if (IsForeignProject)
++			{
++				UProjectPath = String.Format("\"{0}\"", ProjectTarget.UnrealProjectFilePath!.FullName);
++			}
++
++			PlatformProjectGenerator? ProjGenerator = PlatformProjectGenerators.GetPlatformProjectGenerator(Platform, true);
++			VCProjectFile.BuildCommandBuilder BuildCommandBuilder = new(Configuration, Platform, ProjectTarget, UProjectPath)
++			{
++				ProjectGenerator = ProjGenerator,
++				bIsForeignProject = IsForeignProject
++			};
++
++			string BuildArguments = BuildCommandBuilder.GetBuildArguments();
++
++			return new TargetBuildInfo()
++			{
++				BuildCmd = $"{BuildCommandBuilder.BuildScript.FullName} {BuildArguments}",
++				RebuildCmd = $"{BuildCommandBuilder.RebuildScript.FullName} {BuildArguments}",
++				CleanCmd = $"{BuildCommandBuilder.CleanScript.FullName} {BuildArguments}",
++				PrimaryOutput = Target.Binaries[0].OutputFilePath.FullName,
++				BuildByDefault = bBuildByDefault,
++			};
++		}
++
++		internal class ExportedTargetInfo
++		{
++			public string TargetName { get; set; } = String.Empty;
++			public string TargetPath { get; set; } = String.Empty;
++			public string ProjectPath { get; set; } = String.Empty;
++			public string TargetType { get; set; } = String.Empty;
++			public string Platform { get; set; } = String.Empty;
++			public string Configuration { get; set; } = String.Empty;
++			public TargetBuildInfo? BuildInfo { get; set; }
++			public Dictionary<string, ExportedModuleInfo> ModuleToCompileSettings { get; set; } = new();
++			public Dictionary<string, string> DirToModule { get; set; } = new();
++		}
++
++		internal class ExportedModuleInfo
++		{
++			public string Name { get; set; } = String.Empty;
++			public string Directory { get; set; } = String.Empty;
++			public string Rules { get; set; } = String.Empty;
++			public string GeneratedCodeDirectory { get; set; } = String.Empty;
++			public List<string> IncludePaths { get; set; } = new();
++			public List<string> Defines { get; set; } = new();
++			public string? Standard { get; set; }
++			public List<string> ForcedIncludes { get; set; } = new();
++			public string? CompilerPath { get; set; }
++			public List<string> CompilerArgs { get; set; } = new();
++			public Dictionary<string, List<string>> CompilerAdditionalArgs { get; set; } = new();
++			public string? WindowsSdkVersion { get; set; }
++		}
++
++		internal class TargetBuildInfo
++		{
++			public string BuildCmd { get; set; } = String.Empty;
++			public string RebuildCmd { get; set; } = String.Empty;
++			public string CleanCmd { get; set; } = String.Empty;
++			public string PrimaryOutput { get; set; } = String.Empty;
++			public bool BuildByDefault { get; internal set; }
++		}
++	}
++}
+diff --git a/Engine/Source/Programs/UnrealBuildTool/ProjectFiles/VisualStudioWorkspace/VSWorkspaceProjectFileGenerator.cs b/Engine/Source/Programs/UnrealBuildTool/ProjectFiles/VisualStudioWorkspace/VSWorkspaceProjectFileGenerator.cs
+new file mode 100644
+index 0000000000000..22cb46902cc84
+--- /dev/null
++++ b/Engine/Source/Programs/UnrealBuildTool/ProjectFiles/VisualStudioWorkspace/VSWorkspaceProjectFileGenerator.cs
+@@ -0,0 +1,258 @@
++// Copyright Epic Games, Inc. All Rights Reserved.
++
++using System;
++using System.Collections.Generic;
++using System.IO;
++using System.Linq;
++using System.Text.Json;
++using EpicGames.Core;
++using Microsoft.Extensions.Logging;
++using UnrealBuildBase;
++
++namespace UnrealBuildTool
++{
++	class VSWorkspaceProjectFileGenerator : ProjectFileGenerator
++	{
++		public override string ProjectFileExtension => ".json";
++
++		// These properties are used by Visual Studio to determine where to read the project files.
++		// So they must remain constant.
++		private const string VSUnrealWorkspaceFileName = ".vs-unreal-workspace";
++		private const string ProjectFilesFolder = "VisualStudio";
++
++		private readonly CommandLineArguments Arguments;
++
++		/// <summary>
++		/// List of deprecated platforms.
++		/// Don't generate project model for these platforms unless they are specified in "Platforms" console arguments. 
++		/// </summary>
++		/// <returns></returns>
++		private readonly HashSet<UnrealTargetPlatform> DeprecatedPlatforms = new();
++
++		/// <summary>
++		/// Platforms to generate project files for
++		/// </summary>
++		[CommandLine("-Platforms=", ListSeparator = '+')]
++		HashSet<UnrealTargetPlatform> Platforms = new();
++
++		/// <summary>
++		/// Target types to generate project files for
++		/// </summary>
++		[CommandLine("-TargetTypes=", ListSeparator = '+')]
++		HashSet<TargetType> TargetTypes = new();
++
++		/// <summary>
++		/// Target configurations to generate project files for
++		/// </summary>
++		[CommandLine("-TargetConfigurations=", ListSeparator = '+')]
++		HashSet<UnrealTargetConfiguration> TargetConfigurations = new();
++
++		/// <summary>
++		/// Projects to generate project files for
++		/// </summary>
++		[CommandLine("-ProjectNames=", ListSeparator = '+')]
++		HashSet<string> ProjectNames = new();
++
++		/// <summary>
++		/// Should format JSON files in human readable form, or use packed one without indents
++		/// </summary>
++		[CommandLine("-Minimize", Value = "Compact")]
++		private JsonWriterStyle Minimize = JsonWriterStyle.Readable;
++
++		public VSWorkspaceProjectFileGenerator(FileReference? InOnlyGameProject,
++			CommandLineArguments InArguments)
++			: base(InOnlyGameProject)
++		{
++			Arguments = InArguments;
++			Arguments.ApplyTo(this);
++		}
++
++		public override bool ShouldGenerateIntelliSenseData() => true;
++
++		public override void CleanProjectFiles(DirectoryReference InPrimaryProjectDirectory, string InPrimaryProjectName,
++			DirectoryReference InIntermediateProjectFilesDirectory, ILogger Logger)
++		{
++			DirectoryReference.Delete(InPrimaryProjectDirectory);
++		}
++
++		protected override void ConfigureProjectFileGeneration(string[] Arguments, ref bool IncludeAllPlatforms, ILogger Logger)
++		{
++			base.ConfigureProjectFileGeneration(Arguments, ref IncludeAllPlatforms, Logger);
++		}
++
++		protected override ProjectFile AllocateProjectFile(FileReference InitFilePath, DirectoryReference BaseDir)
++		{
++			VSWorkspaceProjectFile projectFile = new(InitFilePath, BaseDir, RootPath: InitFilePath.Directory,
++				Arguments: Arguments, TargetTypes: TargetTypes);
++			return projectFile;
++		}
++
++		protected override bool WriteProjectFiles(PlatformProjectGeneratorCollection PlatformProjectGenerators, ILogger Logger)
++		{
++			using ProgressWriter Progress = new("Writing project files...", true, Logger);
++			List<ProjectFile> ProjectsToGenerate = new(GeneratedProjectFiles);
++			if (ProjectNames.Any())
++			{
++				ProjectsToGenerate = ProjectsToGenerate.Where(Project =>
++					ProjectNames.Contains(Project.ProjectFilePath.GetFileNameWithoutAnyExtensions())).ToList();
++			}
++
++			int TotalProjectFileCount = ProjectsToGenerate.Count;
++
++			HashSet<UnrealTargetPlatform> PlatformsToGenerate = new(SupportedPlatforms);
++			if (Platforms.Any())
++			{
++				PlatformsToGenerate.IntersectWith(Platforms);
++			}
++
++			List<UnrealTargetPlatform> FilteredPlatforms = PlatformsToGenerate.Where(Platform =>
++			{
++				// Skip deprecated unless explicitly specified in the command line.
++				return (!DeprecatedPlatforms.Contains(Platform) || Platforms.Contains(Platform))
++					   && UEBuildPlatform.IsPlatformAvailable(Platform);
++			}).ToList();
++
++			HashSet<UnrealTargetConfiguration> ConfigurationsToGenerate = new(SupportedConfigurations);
++			if (TargetConfigurations.Any())
++			{
++				ConfigurationsToGenerate.IntersectWith(TargetConfigurations);
++			}
++
++			List<UnrealTargetConfiguration> Configurations = ConfigurationsToGenerate.ToList();
++
++			for (int ProjectFileIndex = 0; ProjectFileIndex < ProjectsToGenerate.Count; ++ProjectFileIndex)
++			{
++				if (ProjectsToGenerate[ProjectFileIndex] is not VSWorkspaceProjectFile CurrentProject)
++				{
++					return false;
++				}
++
++				if (!CurrentProject.WriteProjectFile(FilteredPlatforms, Configurations, PlatformProjectGenerators, Minimize, Logger))
++				{
++					return false;
++				}
++
++				Progress.Write(ProjectFileIndex + 1, TotalProjectFileCount);
++			}
++
++			Progress.Write(TotalProjectFileCount, TotalProjectFileCount);
++
++			return true;
++		}
++
++		public override bool GenerateProjectFiles(PlatformProjectGeneratorCollection PlatformProjectGenerators,
++			String[] Arguments, bool bCacheDataForEditor, ILogger Logger)
++		{
++			bool IncludeAllPlatforms = true;
++			ConfigureProjectFileGeneration(Arguments, ref IncludeAllPlatforms, Logger);
++
++			if (bGeneratingGameProjectFiles || Unreal.IsEngineInstalled())
++			{
++				PrimaryProjectPath = OnlyGameProject!.Directory;
++				PrimaryProjectName = OnlyGameProject.GetFileNameWithoutExtension();
++
++				IntermediateProjectFilesPath =
++					DirectoryReference.Combine(PrimaryProjectPath, "Intermediate", "ProjectFiles");
++			}
++
++			SetupSupportedPlatformsAndConfigurations(IncludeAllPlatforms: true, Logger, out string SupportedPlatformNames);
++			Logger.LogDebug("Supported platforms: {Platforms}", SupportedPlatformNames);
++
++			List<FileReference> AllGames = FindGameProjects(Logger);
++
++			{
++				// Find all of the target files.
++				List<FileReference> AllTargetFiles = DiscoverTargets(
++					AllGames,
++					Logger,
++					OnlyGameProject,
++					SupportedPlatforms,
++					bIncludeEngineSource: bIncludeEngineSource,
++					bIncludeTempTargets: bIncludeTempTargets);
++
++				// If there are multiple targets of a given type for a project, use the order to determine which one gets a suffix.
++				AllTargetFiles = AllTargetFiles.OrderBy(x => x.FullName, StringComparer.OrdinalIgnoreCase).ToList();
++
++				List<ProjectFile> EngineProject = new();
++				List<ProjectFile> GameProjects = new();
++				List<ProjectFile> ModProjects = new();
++				Dictionary<FileReference, ProjectFile> ProgramProjects = new();
++				Dictionary<RulesAssembly, DirectoryReference> RulesAssemblies = new();
++				Dictionary<ProjectFile, FileReference> ProjectFileToUProjectFile = new();
++
++				AddProjectsForAllTargets(
++					PlatformProjectGenerators,
++					AllGames,
++					AllTargetFiles,
++					Arguments,
++					EngineProject,
++					GameProjects,
++					ProjectFileToUProjectFile,
++					ProgramProjects,
++					RulesAssemblies,
++					Logger);
++
++				AddAllGameProjects(GameProjects);
++			}
++
++			WriteProjectFiles(PlatformProjectGenerators, Logger);
++			WritePrimaryProjectFile(UBTProject, PlatformProjectGenerators, Logger);
++
++			return true;
++		}
++
++		protected override bool WritePrimaryProjectFile(ProjectFile? UBTProject,
++			PlatformProjectGeneratorCollection PlatformProjectGenerators,
++			ILogger Logger)
++		{
++			try
++			{
++				FileReference PrimaryProjectFile = FileReference.Combine(
++					IntermediateProjectFilesPath, ProjectFilesFolder, VSUnrealWorkspaceFileName);
++
++				DirectoryReference.CreateDirectory(PrimaryProjectFile.Directory);
++
++				// Collect all the resulting project files and aggregate the target-level data
++				var AggregatedProjectInfo = GeneratedProjectFiles
++					.Where(Project => Project is VSWorkspaceProjectFile)
++					.OfType<VSWorkspaceProjectFile>()
++					.SelectMany(Project => Project.ExportedTargetProjects)
++					.GroupBy(TargetProject => TargetProject.TargetName)
++					.Select(g => (g.Key, Target: new
++					{
++						TargetType = g.Select(i => i.TargetType).Distinct().Single(),
++						TargetPath = g.Select(i => i.TargetPath).Distinct().Single(),
++						ProjectPath = g.Select(i => i.ProjectPath).Distinct().Single(),
++						Configurations = g.Select(i => i.Configuration).Distinct().ToList(),
++						Platforms = g.Select(i => i.Platform).Distinct().ToList(),
++					}));
++
++				// The inner Targets object is needed for schema compatibility with the Query Mode API.
++				var Result = new
++				{
++					Targets = AggregatedProjectInfo.ToDictionary(item => item.Key, item => item.Target)
++				};
++
++				using FileStream Stream = new(PrimaryProjectFile.FullName, FileMode.Create, FileAccess.Write);
++				JsonSerializer.Serialize(Stream, Result, new JsonSerializerOptions
++				{
++					PropertyNamingPolicy = JsonNamingPolicy.CamelCase,
++					WriteIndented = Minimize == JsonWriterStyle.Readable,
++				});
++			}
++			catch (Exception Ex)
++			{
++				Logger.LogWarning("Exception while writing root project file: {Ex}", Ex.ToString());
++				return false;
++			}
++
++			return true;
++		}
++
++		/// <inheritdoc />
++		protected override FileReference GetProjectLocation(string BaseName)
++		{
++			return FileReference.Combine(IntermediateProjectFilesPath, ProjectFilesFolder, BaseName + ProjectFileExtension);
++		}
++	}
++}
+diff --git a/Engine/Source/Programs/UnrealBuildTool/ToolChain/UEToolChain.cs b/Engine/Source/Programs/UnrealBuildTool/ToolChain/UEToolChain.cs
+index a6c6f032d394a..0aa18d0b7969e 100644
+--- a/Engine/Source/Programs/UnrealBuildTool/ToolChain/UEToolChain.cs
++++ b/Engine/Source/Programs/UnrealBuildTool/ToolChain/UEToolChain.cs
+@@ -183,6 +183,21 @@ public static FileReference GetResponseFileName(CppCompileEnvironment CompileEnv
+ 			return OutputFile.Location.ChangeExtension(OutputFile.Location.GetExtension() + ResponseExt);
+ 		}
+ 
++		public virtual IEnumerable<string> GetGlobalCommandLineArgs(CppCompileEnvironment CompileEnvironment)
++		{
++			return Array.Empty<string>();
++		}
++
++		public virtual IEnumerable<string> GetCPPCommandLineArgs(CppCompileEnvironment CompileEnvironment)
++		{
++			return Array.Empty<string>();
++		}
++
++		public virtual IEnumerable<string> GetCCommandLineArgs(CppCompileEnvironment CompileEnvironment)
++		{
++			return Array.Empty<string>();
++		}
++
+ 		/// <summary>
+ 		/// Get the name of the response file for the current linker environment and output file
+ 		/// </summary>
