Update 4.27 VSWorkspace patch

#rnx

diff --git a/Engine/Extras/VisualStudioWorkspace/UnrealBuildTool-4.27.patch b/Engine/Extras/VisualStudioWorkspace/UnrealBuildTool-4.27.patch
--- a/Engine/Extras/VisualStudioWorkspace/UnrealBuildTool-4.27.patch
+++ b/Engine/Extras/VisualStudioWorkspace/UnrealBuildTool-4.27.patch
@@ -1,4 +1,4 @@
-From 1ef457067fb87ae7c2317572080b76bd4435c4ac Mon Sep 17 00:00:00 2001
+From 4fc57fd994de108383f838bcbb9af9cef980cbb1 Mon Sep 17 00:00:00 2001
 From: [REDACTED_PERSON] <[REDACTED_EMAIL]>
 Date: Wed, 17 Jan 2024 19:49:54 -0800
 Subject: [PATCH] [UE4.27] Project Generator for VS Workspace support
@@ -18,11 +18,11 @@ open a `.uproject` file directly.
  .../Modes/GenerateProjectFilesMode.cs         |   4 +
  .../Platform/Windows/VCToolChain.cs           |  26 ++
  .../ProjectFiles/ProjectFileGenerator.cs      |   5 +-
- .../VSWorkspaceProjectFile.cs                 | 388 ++++++++++++++++++
+ .../VSWorkspaceProjectFile.cs                 | 392 ++++++++++++++++++
  .../VSWorkspaceProjectFileGenerator.cs        | 277 +++++++++++++
  .../UnrealBuildTool/ToolChain/UEToolChain.cs  |  25 ++
  .../UnrealBuildTool/UnrealBuildTool.csproj    |   2 +
- 8 files changed, 726 insertions(+), 3 deletions(-)
+ 8 files changed, 730 insertions(+), 3 deletions(-)
  create mode 100644 Engine/Source/Programs/UnrealBuildTool/ProjectFiles/VisualStudioWorkspace/VSWorkspaceProjectFile.cs
  create mode 100644 Engine/Source/Programs/UnrealBuildTool/ProjectFiles/VisualStudioWorkspace/VSWorkspaceProjectFileGenerator.cs
 
@@ -137,10 +137,10 @@ index 66a1ec6d0b227..f64a010b97581 100644
  		}
 diff --git a/Engine/Source/Programs/UnrealBuildTool/ProjectFiles/VisualStudioWorkspace/VSWorkspaceProjectFile.cs b/Engine/Source/Programs/UnrealBuildTool/ProjectFiles/VisualStudioWorkspace/VSWorkspaceProjectFile.cs
 new file mode 100644
-index 0000000000000..2e37959f34c0a
+index 0000000000000..fb5dffb1ef446
 --- /dev/null
 +++ b/Engine/Source/Programs/UnrealBuildTool/ProjectFiles/VisualStudioWorkspace/VSWorkspaceProjectFile.cs
-@@ -0,0 +1,388 @@
+@@ -0,0 +1,392 @@
 +// Copyright Epic Games, Inc. All Rights Reserved.
 +
 +using System;
@@ -209,8 +209,8 @@ index 0000000000000..2e37959f34c0a
 +						}
 +
 +						// Skip Editor for all platforms except for current platform
-+						if (ProjectTarget.TargetRules.Type == TargetType.Editor && 
-+							(BuildHostPlatform.Current.Platform != Platform || 
++						if (ProjectTarget.TargetRules.Type == TargetType.Editor &&
++							(BuildHostPlatform.Current.Platform != Platform ||
 +							(Configuration == UnrealTargetConfiguration.Test || Configuration == UnrealTargetConfiguration.Shipping)))
 +						{
 +							continue;
@@ -227,7 +227,7 @@ index 0000000000000..2e37959f34c0a
 +
 +						try
 +						{
-+							FileReference OutputFile = FileReference.Combine(ProjectRootFolder, 
++							FileReference OutputFile = FileReference.Combine(ProjectRootFolder,
 +								$"{ProjectTarget.TargetFilePath.GetFileNameWithoutAnyExtensions()}_{Configuration}_{Platform}.json");
 +
 +							UEBuildTarget BuildTarget = UEBuildTarget.Create(TargetDesc, false, false);
@@ -235,7 +235,7 @@ index 0000000000000..2e37959f34c0a
 +						}
 +						catch (Exception Ex)
 +						{
-+							Log.TraceWarning("Exception while generating include data for Target:{0}, Platform: {1}, Configuration: {2}", 
++							Log.TraceWarning("Exception while generating include data for Target:{0}, Platform: {1}, Configuration: {2}",
 +								TargetDesc.Name, Platform.ToString(), Configuration.ToString());
 +							Log.TraceWarning(Ex.ToString());
 +						}
@@ -255,9 +255,13 @@ index 0000000000000..2e37959f34c0a
 +
 +					ExportedTargetInfo TargetInfo = ExportTarget(BuildTarget, bBuildByDefault, PlatformProjectGenerators);
 +
++					// Need a workaround to escape the keys of the DirToModule dictionary, which is not done automatically
++					// when serializing, resulting in an invalid JSON content.
++					TargetInfo.DirToModule = TargetInfo.DirToModule.ToDictionary(x => Json.EscapeString(x.Key), x => x.Value);
++
 +					DirectoryReference.CreateDirectory(OutputFile.Directory);
 +
-+					JsonSerializeOptions Options = Minimize == JsonWriterStyle.Readable 
++					JsonSerializeOptions Options = Minimize == JsonWriterStyle.Readable
 +						? JsonSerializeOptions.PrettyPrint : JsonSerializeOptions.None;
 +					File.WriteAllText(OutputFile.FullName, Json.Serialize(TargetInfo, Options));
 +
@@ -317,10 +321,10 @@ index 0000000000000..2e37959f34c0a
 +
 +					if (ModuleCpp.GeneratedCodeDirectory != null)
 +					{
-+                        if (!TargetInfo.DirToModule.ContainsKey(ModuleCpp.GeneratedCodeDirectory.FullName))
-+                        {
-+                            TargetInfo.DirToModule.Add(ModuleCpp.GeneratedCodeDirectory.FullName, ModuleCpp.Name);
-+                        }
++						if (!TargetInfo.DirToModule.ContainsKey(ModuleCpp.GeneratedCodeDirectory.FullName))
++						{
++							TargetInfo.DirToModule.Add(ModuleCpp.GeneratedCodeDirectory.FullName, ModuleCpp.Name);
++						}
 +					}
 +				}
 +			}
